Fix warning caused by struct forward-declared as a class.

#rnx

diff --git a/Engine/Plugins/Experimental/ProxyLODPlugin/Source/ProxyLOD/Private/ProxyLODkDOPInterface.h b/Engine/Plugins/Experimental/ProxyLODPlugin/Source/ProxyLOD/Private/ProxyLODkDOPInterface.h
--- a/Engine/Plugins/Experimental/ProxyLODPlugin/Source/ProxyLOD/Private/ProxyLODkDOPInterface.h
+++ b/Engine/Plugins/Experimental/ProxyLODPlugin/Source/ProxyLOD/Private/ProxyLODkDOPInterface.h
@@ -7,7 +7,7 @@
 #include "ProxyLODMeshTypes.h"
 #include "ProxyLODThreadedWrappers.h"
 
-class FMeshDescription;
+struct FMeshDescription;
 
 namespace ProxyLOD
 {
