Remove overrides forcing undefined identifier and shadow variable warnings off. Separating out from change to force these settings to create a new PCH to prevent the warning setting being ignored.

diff --git a/Engine/Plugins/Experimental/ProxyLODPlugin/Source/ProxyLOD/Private/ProxyLODMaterialTransferUtilities.h b/Engine/Plugins/Experimental/ProxyLODPlugin/Source/ProxyLOD/Private/ProxyLODMaterialTransferUtilities.h
--- a/Engine/Plugins/Experimental/ProxyLODPlugin/Source/ProxyLOD/Private/ProxyLODMaterialTransferUtilities.h
+++ b/Engine/Plugins/Experimental/ProxyLODPlugin/Source/ProxyLOD/Private/ProxyLODMaterialTransferUtilities.h
@@ -8,7 +8,9 @@
 #include "ProxyLODGrid2d.h"
 #include "ProxyLODRasterizer.h" // FRasterGrid
 
+THIRD_PARTY_INCLUDES_START
 #include <openvdb/openvdb.h>
+THIRD_PARTY_INCLUDES_END
 
 namespace ProxyLOD
 {
diff --git a/Engine/Plugins/Experimental/ProxyLODPlugin/Source/ProxyLOD/Private/ProxyLODMeshAttrTransfer.h b/Engine/Plugins/Experimental/ProxyLODPlugin/Source/ProxyLOD/Private/ProxyLODMeshAttrTransfer.h
--- a/Engine/Plugins/Experimental/ProxyLODPlugin/Source/ProxyLOD/Private/ProxyLODMeshAttrTransfer.h
+++ b/Engine/Plugins/Experimental/ProxyLODPlugin/Source/ProxyLOD/Private/ProxyLODMeshAttrTransfer.h
@@ -5,7 +5,9 @@
 #include "ProxyLODMeshTypes.h"
 #include "ProxyLODThreadedWrappers.h"
 
+THIRD_PARTY_INCLUDES_START
 #include <openvdb/openvdb.h>
+THIRD_PARTY_INCLUDES_END
 
 namespace ProxyLOD
 {
diff --git a/Engine/Plugins/Experimental/ProxyLODPlugin/Source/ProxyLOD/Private/ProxyLODMeshParameterization.cpp b/Engine/Plugins/Experimental/ProxyLODPlugin/Source/ProxyLOD/Private/ProxyLODMeshParameterization.cpp
--- a/Engine/Plugins/Experimental/ProxyLODPlugin/Source/ProxyLOD/Private/ProxyLODMeshParameterization.cpp
+++ b/Engine/Plugins/Experimental/ProxyLODPlugin/Source/ProxyLOD/Private/ProxyLODMeshParameterization.cpp
@@ -2,8 +2,10 @@
 
 #include "ProxyLODMeshParameterization.h"
 
+THIRD_PARTY_INCLUDES_START
 #include <UVAtlasCode/UVAtlas/inc/UVAtlas.h>
 #include <DirectXMeshCode/DirectXMesh/DirectXMesh.h>
+THIRD_PARTY_INCLUDES_END
 
 
 #include "ProxyLODThreadedWrappers.h"
diff --git a/Engine/Plugins/Experimental/ProxyLODPlugin/Source/ProxyLOD/Private/ProxyLODMeshPartition.h b/Engine/Plugins/Experimental/ProxyLODPlugin/Source/ProxyLOD/Private/ProxyLODMeshPartition.h
--- a/Engine/Plugins/Experimental/ProxyLODPlugin/Source/ProxyLOD/Private/ProxyLODMeshPartition.h
+++ b/Engine/Plugins/Experimental/ProxyLODPlugin/Source/ProxyLOD/Private/ProxyLODMeshPartition.h
@@ -119,7 +119,7 @@ namespace ProxyLOD
 // --- Implementation of templated functions --- 
 
 	template <typename VertexType, typename PartitionFunctor>
-	static void ProxyLOD::CreatePartitionArray( const TAOSMesh<VertexType>& Mesh, 
+	void ProxyLOD::CreatePartitionArray( const TAOSMesh<VertexType>& Mesh, 
 		                                        const PartitionFunctor& Functor, 
 		                                        TArray<int32>& PartitionArray)
 	{
@@ -153,7 +153,7 @@ namespace ProxyLOD
 
 
 	template <typename VertexType, typename PartitionFunctor>
-	static void ProxyLOD::PartitionMeshSplit( const TAOSMesh<VertexType>& SrcMesh, 
+	void ProxyLOD::PartitionMeshSplit( const TAOSMesh<VertexType>& SrcMesh, 
 		                                      const PartitionFunctor& Functor, 
 		                                      TArray<TAOSMesh<VertexType>>& ResultMeshArray)
 	{
@@ -290,7 +290,7 @@ namespace ProxyLOD
 
 	
 	template <typename VertexType>
-	static bool  ProxyLOD::MergeMeshArray( TArray<TAOSMesh<VertexType>>& MeshArray, TAOSMesh<VertexType>& ResultMesh)
+	bool  ProxyLOD::MergeMeshArray( TArray<TAOSMesh<VertexType>>& MeshArray, TAOSMesh<VertexType>& ResultMesh)
 	{
 		// Determine the space needed.
 
@@ -344,7 +344,7 @@ namespace ProxyLOD
 
 
 	template <typename VertexType>
-	static bool  ProxyLOD::MergeMeshArray( TArray<TAOSMesh<VertexType>>& MeshArray, 
+	bool  ProxyLOD::MergeMeshArray( TArray<TAOSMesh<VertexType>>& MeshArray, 
 		                                   const TArray<int32> * SeamVertArrayOfArrays, 
 		                                   TAOSMesh<VertexType>& ResultMesh)
 	{
diff --git a/Engine/Plugins/Experimental/ProxyLODPlugin/Source/ProxyLOD/Private/ProxyLODMeshSDFConversions.cpp b/Engine/Plugins/Experimental/ProxyLODPlugin/Source/ProxyLOD/Private/ProxyLODMeshSDFConversions.cpp
--- a/Engine/Plugins/Experimental/ProxyLODPlugin/Source/ProxyLOD/Private/ProxyLODMeshSDFConversions.cpp
+++ b/Engine/Plugins/Experimental/ProxyLODPlugin/Source/ProxyLOD/Private/ProxyLODMeshSDFConversions.cpp
@@ -5,9 +5,13 @@
 
 #include "CoreMinimal.h"
 
+THIRD_PARTY_INCLUDES_START
+#pragma warning(push)
+#pragma warning(disable: 4146)
 #include <openvdb/tools/MeshToVolume.h> // for MeshToVolume
-
 #include <openvdb/tools/Composite.h> // for csgUnion
+#pragma warning(pop)
+THIRD_PARTY_INCLUDES_END
 
 //#include <exception>
 
diff --git a/Engine/Plugins/Experimental/ProxyLODPlugin/Source/ProxyLOD/Private/ProxyLODMeshSDFConversions.h b/Engine/Plugins/Experimental/ProxyLODPlugin/Source/ProxyLOD/Private/ProxyLODMeshSDFConversions.h
--- a/Engine/Plugins/Experimental/ProxyLODPlugin/Source/ProxyLOD/Private/ProxyLODMeshSDFConversions.h
+++ b/Engine/Plugins/Experimental/ProxyLODPlugin/Source/ProxyLOD/Private/ProxyLODMeshSDFConversions.h
@@ -6,12 +6,10 @@
 #include "ProxyLODMeshTypes.h"
 #include "ProxyLODMeshConvertUtils.h"
 
+THIRD_PARTY_INCLUDES_START
 #include <openvdb/openvdb.h>
-
-
-
 #include <openvdb/tools/VolumeToMesh.h> // for VolumeToMesh
-
+THIRD_PARTY_INCLUDES_END
 
 namespace ProxyLOD
 {
diff --git a/Engine/Plugins/Experimental/ProxyLODPlugin/Source/ProxyLOD/Private/ProxyLODMeshTypes.h b/Engine/Plugins/Experimental/ProxyLODPlugin/Source/ProxyLOD/Private/ProxyLODMeshTypes.h
--- a/Engine/Plugins/Experimental/ProxyLODPlugin/Source/ProxyLOD/Private/ProxyLODMeshTypes.h
+++ b/Engine/Plugins/Experimental/ProxyLODPlugin/Source/ProxyLOD/Private/ProxyLODMeshTypes.h
@@ -19,9 +19,10 @@
 #include "ProxyLODThreadedWrappers.h"
 #include "ProxyLODVertexTypes.h"
 
+THIRD_PARTY_INCLUDES_START
 #include <openvdb/openvdb.h>
 #include <vector>
-
+THIRD_PARTY_INCLUDES_END
 
 /**
 * Various stages in the ProxyLOD pipeline require different internal mesh formats, but
diff --git a/Engine/Plugins/Experimental/ProxyLODPlugin/Source/ProxyLOD/Private/ProxyLODMeshUtilities.cpp b/Engine/Plugins/Experimental/ProxyLODPlugin/Source/ProxyLOD/Private/ProxyLODMeshUtilities.cpp
--- a/Engine/Plugins/Experimental/ProxyLODPlugin/Source/ProxyLOD/Private/ProxyLODMeshUtilities.cpp
+++ b/Engine/Plugins/Experimental/ProxyLODPlugin/Source/ProxyLOD/Private/ProxyLODMeshUtilities.cpp
@@ -10,7 +10,9 @@
 #include "Modules/ModuleManager.h"
 #include "MeshUtilities.h" // IMeshUtilities
 
+THIRD_PARTY_INCLUDES_START
 #include <DirectXMeshCode/DirectXMesh/DirectXMesh.h>
+THIRD_PARTY_INCLUDES_END
 
 #include "MeshDescription.h"
 #include "MeshAttributes.h"
diff --git a/Engine/Plugins/Experimental/ProxyLODPlugin/Source/ProxyLOD/Private/ProxyLODMeshUtilities.h b/Engine/Plugins/Experimental/ProxyLODPlugin/Source/ProxyLOD/Private/ProxyLODMeshUtilities.h
--- a/Engine/Plugins/Experimental/ProxyLODPlugin/Source/ProxyLOD/Private/ProxyLODMeshUtilities.h
+++ b/Engine/Plugins/Experimental/ProxyLODPlugin/Source/ProxyLOD/Private/ProxyLODMeshUtilities.h
@@ -7,7 +7,9 @@
 #include "ProxyLODMeshTypes.h"
 #include "ProxyLODThreadedWrappers.h"
 
+THIRD_PARTY_INCLUDES_START
 #include <vector>
+THIRD_PARTY_INCLUDES_END
 
 namespace ProxyLOD
 {
diff --git a/Engine/Plugins/Experimental/ProxyLODPlugin/Source/ProxyLOD/Private/ProxyLODPlugin.cpp b/Engine/Plugins/Experimental/ProxyLODPlugin/Source/ProxyLOD/Private/ProxyLODPlugin.cpp
--- a/Engine/Plugins/Experimental/ProxyLODPlugin/Source/ProxyLOD/Private/ProxyLODPlugin.cpp
+++ b/Engine/Plugins/Experimental/ProxyLODPlugin/Source/ProxyLOD/Private/ProxyLODPlugin.cpp
@@ -12,8 +12,10 @@
 #include "Misc/ScopedSlowTask.h"
 #include "Stats/StatsMisc.h"
 //#include "ScopedTimers.h"
+THIRD_PARTY_INCLUDES_START
 #include <openvdb/openvdb.h>
 #include <openvdb/tools/Interpolation.h> // for grid sampler
+THIRD_PARTY_INCLUDES_END
 
 #define PROXYLOD_CLOCKWISE_TRIANGLES  1
 
diff --git a/Engine/Plugins/Experimental/ProxyLODPlugin/Source/ProxyLOD/Private/ProxyLODRasterizer.cpp b/Engine/Plugins/Experimental/ProxyLODPlugin/Source/ProxyLOD/Private/ProxyLODRasterizer.cpp
--- a/Engine/Plugins/Experimental/ProxyLODPlugin/Source/ProxyLOD/Private/ProxyLODRasterizer.cpp
+++ b/Engine/Plugins/Experimental/ProxyLODPlugin/Source/ProxyLOD/Private/ProxyLODRasterizer.cpp
@@ -6,7 +6,9 @@
 #include "ProxyLODMeshUtilities.h"
 #include "ProxyLODThreadedWrappers.h"
 
+THIRD_PARTY_INCLUDES_START
 #include <tbb/spin_mutex.h>
+THIRD_PARTY_INCLUDES_END
 
 namespace ProxyLOD
 {
diff --git a/Engine/Plugins/Experimental/ProxyLODPlugin/Source/ProxyLOD/Private/ProxyLODRasterizer.h b/Engine/Plugins/Experimental/ProxyLODPlugin/Source/ProxyLOD/Private/ProxyLODRasterizer.h
--- a/Engine/Plugins/Experimental/ProxyLODPlugin/Source/ProxyLOD/Private/ProxyLODRasterizer.h
+++ b/Engine/Plugins/Experimental/ProxyLODPlugin/Source/ProxyLOD/Private/ProxyLODRasterizer.h
@@ -7,10 +7,11 @@
 #include "Math/UnrealMathUtility.h" // ComputeBaryCentric2D
 
 // std
+THIRD_PARTY_INCLUDES_START
 #include <limits>
 #include <array>
 #include <tuple>
-
+THIRD_PARTY_INCLUDES_END
 
 // openvdb
 //#include <openvdb/openvdb.h>
diff --git a/Engine/Plugins/Experimental/ProxyLODPlugin/Source/ProxyLOD/Private/ProxyLODThreadedWrappers.h b/Engine/Plugins/Experimental/ProxyLODPlugin/Source/ProxyLOD/Private/ProxyLODThreadedWrappers.h
--- a/Engine/Plugins/Experimental/ProxyLODPlugin/Source/ProxyLOD/Private/ProxyLODThreadedWrappers.h
+++ b/Engine/Plugins/Experimental/ProxyLODPlugin/Source/ProxyLOD/Private/ProxyLODThreadedWrappers.h
@@ -1,6 +1,7 @@
 // Copyright 1998-2019 Epic Games, Inc. All Rights Reserved.
 #pragma once
 
+THIRD_PARTY_INCLUDES_START
 PRAGMA_DISABLE_DEPRECATION_WARNINGS
 #include <tbb/blocked_range.h>
 #include <tbb/blocked_range2d.h>
@@ -9,6 +10,7 @@ PRAGMA_DISABLE_DEPRECATION_WARNINGS
 #include <tbb/parallel_reduce.h>
 #include <tbb/task_group.h>
 PRAGMA_ENABLE_DEPRECATION_WARNINGS
+THIRD_PARTY_INCLUDES_END
 
 /**
 * Wrappers for the tbb calls to allow for easier single threaded testing.
diff --git a/Engine/Plugins/Experimental/ProxyLODPlugin/Source/ProxyLOD/Private/ProxyLODTwoDTriangleUtilities.h b/Engine/Plugins/Experimental/ProxyLODPlugin/Source/ProxyLOD/Private/ProxyLODTwoDTriangleUtilities.h
--- a/Engine/Plugins/Experimental/ProxyLODPlugin/Source/ProxyLOD/Private/ProxyLODTwoDTriangleUtilities.h
+++ b/Engine/Plugins/Experimental/ProxyLODPlugin/Source/ProxyLOD/Private/ProxyLODTwoDTriangleUtilities.h
@@ -2,8 +2,10 @@
 
 #pragma once
 
+THIRD_PARTY_INCLUDES_START
 #include <array>
 #include <limits>
+THIRD_PARTY_INCLUDES_END
 
 namespace ProxyLOD
 {
diff --git a/Engine/Plugins/Experimental/ProxyLODPlugin/Source/ProxyLOD/Private/ProxyLODVolume.cpp b/Engine/Plugins/Experimental/ProxyLODPlugin/Source/ProxyLOD/Private/ProxyLODVolume.cpp
--- a/Engine/Plugins/Experimental/ProxyLODPlugin/Source/ProxyLOD/Private/ProxyLODVolume.cpp
+++ b/Engine/Plugins/Experimental/ProxyLODPlugin/Source/ProxyLOD/Private/ProxyLODVolume.cpp
@@ -8,10 +8,15 @@
 #include "ProxyLODMeshTypes.h"
 #include "ProxyLODMeshUtilities.h"
 
+THIRD_PARTY_INCLUDES_START
+#pragma warning(push)
+#pragma warning(disable: 4146)
 #include <openvdb/openvdb.h>
 #include <openvdb/tools/Interpolation.h> // for Spatial Query
 #include <openvdb/tools/MeshToVolume.h> // for MeshToVolume
 #include <openvdb/tools/VolumeToMesh.h> // for VolumeToMesh
+#pragma warning(pop)
+THIRD_PARTY_INCLUDES_END
 
 #include "MeshDescription.h"
 
diff --git a/Engine/Plugins/Experimental/ProxyLODPlugin/Source/ProxyLOD/ProxyLODMeshReduction.Build.cs b/Engine/Plugins/Experimental/ProxyLODPlugin/Source/ProxyLOD/ProxyLODMeshReduction.Build.cs
--- a/Engine/Plugins/Experimental/ProxyLODPlugin/Source/ProxyLOD/ProxyLODMeshReduction.Build.cs
+++ b/Engine/Plugins/Experimental/ProxyLODPlugin/Source/ProxyLOD/ProxyLODMeshReduction.Build.cs
@@ -9,7 +9,6 @@ namespace UnrealBuildTool.Rules
 
 
             // For boost:: and TBB:: code
-            bEnableUndefinedIdentifierWarnings = false;
             bUseRTTI = true;
 
             PublicIncludePaths.AddRange(
diff --git a/Engine/Source/Developer/VulkanShaderFormat/Private/VulkanShaderCompiler.cpp b/Engine/Source/Developer/VulkanShaderFormat/Private/VulkanShaderCompiler.cpp
--- a/Engine/Source/Developer/VulkanShaderFormat/Private/VulkanShaderCompiler.cpp
+++ b/Engine/Source/Developer/VulkanShaderFormat/Private/VulkanShaderCompiler.cpp
@@ -655,13 +655,13 @@ struct FVulkanHlslccHeader : public CrossCompiler::FHlslccHeader
 			ShaderSource += 22;
 			while (*ShaderSource && *ShaderSource != '\n')
 			{
-				FString Name;
-				if (!CrossCompiler::ParseIdentifier(ShaderSource, Name))
+				FString ExternalTextureName;
+				if (!CrossCompiler::ParseIdentifier(ShaderSource, ExternalTextureName))
 				{
 					return false;
 				}
 
-				ExternalTextures.Add(Name);
+				ExternalTextures.Add(ExternalTextureName);
 
 				if (Match(ShaderSource, '\n'))
 				{
diff --git a/Engine/Source/Developer/VulkanShaderFormat/VulkanShaderFormat.Build.cs b/Engine/Source/Developer/VulkanShaderFormat/VulkanShaderFormat.Build.cs
--- a/Engine/Source/Developer/VulkanShaderFormat/VulkanShaderFormat.Build.cs
+++ b/Engine/Source/Developer/VulkanShaderFormat/VulkanShaderFormat.Build.cs
@@ -20,9 +20,6 @@ public class VulkanShaderFormat : ModuleRules
 			}
 			);
 
-		// GlsLang has a shadowed variable:
-		bEnableShadowVariableWarnings = false;
-
 		AddEngineThirdPartyPrivateStaticDependencies(Target, "HLSLCC");
 		AddEngineThirdPartyPrivateStaticDependencies(Target, "GlsLang");
 
diff --git a/Engine/Source/Programs/UnrealCEFSubProcess/UnrealCEFSubProcess.Build.cs b/Engine/Source/Programs/UnrealCEFSubProcess/UnrealCEFSubProcess.Build.cs
--- a/Engine/Source/Programs/UnrealCEFSubProcess/UnrealCEFSubProcess.Build.cs
+++ b/Engine/Source/Programs/UnrealCEFSubProcess/UnrealCEFSubProcess.Build.cs
@@ -30,8 +30,6 @@ public class UnrealCEFSubProcess : ModuleRules
 		AddEngineThirdPartyPrivateStaticDependencies(Target,
 			"CEF3"
 			);
-			
-		bEnableShadowVariableWarnings = false;
 	}
 }
 
diff --git a/Engine/Source/Runtime/ImageWrapper/ImageWrapper.Build.cs b/Engine/Source/Runtime/ImageWrapper/ImageWrapper.Build.cs
--- a/Engine/Source/Runtime/ImageWrapper/ImageWrapper.Build.cs
+++ b/Engine/Source/Runtime/ImageWrapper/ImageWrapper.Build.cs
@@ -37,8 +37,6 @@ public class ImageWrapper : ModuleRules
 			PublicDefinitions.Add("WITH_UNREALEXR=0");
 		}
 
-		bEnableShadowVariableWarnings = false;
-
 		// Enable exceptions to allow error handling
 		bEnableExceptions = true;
 
diff --git a/Engine/Source/Runtime/WebBrowser/WebBrowser.Build.cs b/Engine/Source/Runtime/WebBrowser/WebBrowser.Build.cs
--- a/Engine/Source/Runtime/WebBrowser/WebBrowser.Build.cs
+++ b/Engine/Source/Runtime/WebBrowser/WebBrowser.Build.cs
@@ -89,7 +89,5 @@ public class WebBrowser : ModuleRules
 		{
 			PrecompileForTargets = ModuleRules.PrecompileTargetsType.None;
 		}
-
-		bEnableShadowVariableWarnings = false;
 	}
 }
