ProxyLOD :  Add early out and warning when input to system does not contain static meshes

diff --git a/Engine/Plugins/Experimental/ProxyLODPlugin/Source/ProxyLOD/Private/ProxyLODPlugin.cpp b/Engine/Plugins/Experimental/ProxyLODPlugin/Source/ProxyLOD/Private/ProxyLODPlugin.cpp
--- a/Engine/Plugins/Experimental/ProxyLODPlugin/Source/ProxyLOD/Private/ProxyLODPlugin.cpp
+++ b/Engine/Plugins/Experimental/ProxyLODPlugin/Source/ProxyLOD/Private/ProxyLODPlugin.cpp
@@ -578,6 +578,22 @@ void FVoxelizeMeshMerging::ProxyLOD(const FMeshMergeDataArray& InData, const FMe
 		}
 	}
 
+	if (InGeometry.Num() == 0)
+	{
+		UE_LOG(LogProxyLODMeshReduction, Warning, TEXT("No static meshes input, no output mesh will be generated."));
+
+		FailedDelegate.ExecuteIfBound(InJobGUID, TEXT("ProxyLOD no input geometry"));
+
+		// Done with the material baking, free the delegate
+
+		if (BakeMaterialsDelegate.IsBound())
+		{
+			BakeMaterialsDelegate.Unbind();
+		}
+
+		return;
+	}
+
 	
 	// Container for the raw mesh that will hold the simplified geometry
 	// and the FlattenMaterial that will hold the materials for the output mesh.
@@ -601,6 +617,23 @@ void FVoxelizeMeshMerging::ProxyLOD(const FMeshMergeDataArray& InData, const FMe
 
 		{
 			const auto& BBox = SrcGeometryAdapter.GetBBox();
+
+			if (!BBox.hasVolume())
+			{
+				UE_LOG(LogProxyLODMeshReduction, Warning, TEXT("Empty bounding box for all static meshes input, no output mesh will be generated."));
+
+				FailedDelegate.ExecuteIfBound(InJobGUID, TEXT("ProxyLOD invalid input geometry"));
+
+				// Done with the material baking, free the delegate
+
+				if (BakeMaterialsDelegate.IsBound())
+				{
+					BakeMaterialsDelegate.Unbind();
+				}
+
+				return;
+			}
+
 			// Determine the voxel size the corresponds to the InProxySettings
 			// The transform defines resolution of the voxel grid.
 
