Deprecate FRawMesh - ProxyLODPlugin now use FMeshDescription

diff --git a/Engine/Plugins/Experimental/ProxyLODPlugin/Source/ProxyLOD/Private/ProxyLODMaterialTransferUtilities.cpp b/Engine/Plugins/Experimental/ProxyLODPlugin/Source/ProxyLOD/Private/ProxyLODMaterialTransferUtilities.cpp
--- a/Engine/Plugins/Experimental/ProxyLODPlugin/Source/ProxyLOD/Private/ProxyLODMaterialTransferUtilities.cpp
+++ b/Engine/Plugins/Experimental/ProxyLODPlugin/Source/ProxyLOD/Private/ProxyLODMaterialTransferUtilities.cpp
@@ -9,6 +9,10 @@
 #include "ProxyLODMeshConvertUtils.h"
 #include "Math/Matrix.h" // used in chirality..
 
+#include "MeshDescription.h"
+#include "MeshAttributes.h"
+#include "MeshAttributeArray.h"
+
 
 namespace ProxyLOD
 {
@@ -19,13 +23,13 @@ namespace
 {
 	// returns the HitUV
 	template<typename TestFunctorType>
-	float EvaluateHit(const FVector& HitPoint, const uint32 SrcPolyId, const FRawMeshArrayAdapter& MeshAdapter,
+	float EvaluateHit(const FVector& HitPoint, const uint32 SrcPolyId, const FMeshDescriptionArrayAdapter& MeshAdapter,
 		ProxyLOD::FSrcMeshData& HitPayload, TestFunctorType& TestFunctor)
 	{
 
 
 		int32 MeshId; int32 LocalFaceNumber;
-		const FRawMeshArrayAdapter::FRawPoly SrcPoly = MeshAdapter.GetRawPoly(SrcPolyId, MeshId, LocalFaceNumber);
+		const FMeshDescriptionArrayAdapter::FRawPoly SrcPoly = MeshAdapter.GetRawPoly(SrcPolyId, MeshId, LocalFaceNumber);
 		int32 MaterialIdx = SrcPoly.FaceMaterialIndex;
 
 		FVector2D UVs[3];
@@ -92,17 +96,17 @@ namespace
 		return TestFunctor(SrcPoly);
 	}
 
-	float EvaluateHit(const FVector& HitPoint, const uint32 SrcPolyId, const FRawMeshArrayAdapter& MeshAdapter,
+	float EvaluateHit(const FVector& HitPoint, const uint32 SrcPolyId, const FMeshDescriptionArrayAdapter& MeshAdapter,
 		ProxyLOD::FSrcMeshData& HitPayload)
 	{
-		auto NoOpFunctor = [](const FRawMeshArrayAdapter::FRawPoly& RawPoly)->float {return 1; };
+		auto NoOpFunctor = [](const FMeshDescriptionArrayAdapter::FRawPoly& RawPoly)->float {return 1; };
 
 		return EvaluateHit(HitPoint, SrcPolyId, MeshAdapter, HitPayload, NoOpFunctor);
 	};
 
 	// returns the HitUV
 	template <bool Forward>
-	float EvaluateHit(const FkHitResult& HitResult, const FVector& RayStart, const FVector& RayDirection, const FRawMeshArrayAdapter& MeshAdapter,
+	float EvaluateHit(const FkHitResult& HitResult, const FVector& RayStart, const FVector& RayDirection, const FMeshDescriptionArrayAdapter& MeshAdapter,
 		ProxyLOD::FSrcMeshData& HitPayload)
 	{
 
@@ -115,7 +119,7 @@ namespace
 		const float  HitTime = (Forward) ? HitResult.Time : -HitResult.Time;
 		const FVector HitPoint = RayStart + HitTime * RayDirection;
 
-		auto DirectionalTest = [&RayDirection](const FRawMeshArrayAdapter::FRawPoly& RawPoly)->float
+		auto DirectionalTest = [&RayDirection](const FMeshDescriptionArrayAdapter::FRawPoly& RawPoly)->float
 		{
 			// Is the SrcNormal roughly in the same direction as the ray
 			// the face normal of this poly
@@ -137,7 +141,7 @@ namespace
 // Generate a map between UV locations on the simplified geometry and points on the Src geometry
 // Inactive texels in the result have MaterialId = -1;
 
-ProxyLOD::FSrcDataGrid::Ptr ProxyLOD::CreateCorrespondence( const FRawMeshArrayAdapter& SrcMesh, 
+ProxyLOD::FSrcDataGrid::Ptr ProxyLOD::CreateCorrespondence( const FMeshDescriptionArrayAdapter& SrcMesh,
 															const FVertexDataMesh& ReducedMesh, 
 															const ProxyLOD::FRasterGrid& UVGrid, 
 															const int32 TransferType, 
@@ -377,7 +381,7 @@ ProxyLOD::FSrcDataGrid::Ptr ProxyLOD::CreateCorrespondence( const FRawMeshArrayA
 // Generate a map between UV locations on the simplified geometry and points on the Src geometry
 // Inactive texels in the result have MaterialId = -1;
 ProxyLOD::FSrcDataGrid::Ptr ProxyLOD::CreateCorrespondence( const openvdb::Int32Grid& ClosestPolyGrid, 
-															const FRawMeshArrayAdapter& SrcMesh,
+															const FMeshDescriptionArrayAdapter& SrcMesh,
 															const FVertexDataMesh& ReducedMesh, 
 															const ProxyLOD::FRasterGrid& UVGrid,
 															const int32 TransferType, 
@@ -1033,7 +1037,7 @@ namespace
 	// In the case that the source geometry didn't have a normal map, we have to generate one by 
 	// directly sampling the geometry normals.
 
-	void SuperSampleWSNormal(const FRawMeshArrayAdapter& SrcMeshAdapter,
+	void SuperSampleWSNormal(const FMeshDescriptionArrayAdapter& SrcMeshAdapter,
 		const ProxyLOD::FSrcDataGrid& SuperSampleCorrespondenceGrid,
 		const ProxyLOD::FRasterGrid& SuperSampleUVGrid,
 		const TArray<FFlattenMaterial>& InputMaterials,
@@ -1073,7 +1077,7 @@ namespace
 			{
 
 				int32 MeshId; int32 LocalFaceNumber;
-				const FRawMeshArrayAdapter::FRawPoly SrcPoly = SrcMeshAdapter.GetRawPoly(TriangleId, MeshId, LocalFaceNumber);
+				const FMeshDescriptionArrayAdapter::FRawPoly SrcPoly = SrcMeshAdapter.GetRawPoly(TriangleId, MeshId, LocalFaceNumber);
 
 				// Get the tangent space:  We try to compute this in a way that is consistent with LocalVertexFactory.usf (CalcTangentToLocal).
 
@@ -1252,10 +1256,14 @@ namespace
 	*      texels correspond to each result texel
 	*  
 	*/
-	void SparseDownSampleNormal(const ProxyLOD::TGridWrapper<FVector>& SuperSampledNormalGrid, 
-		                        const ProxyLOD::FRasterGrid& DstUVGrid, const FRawMesh& DstRawMesh,
-		                        ProxyLOD::FLinearColorGrid& OutDstBufferGrid)
+	void SparseDownSampleNormal(const ProxyLOD::TGridWrapper<FVector>& SuperSampledNormalGrid,
+		const ProxyLOD::FRasterGrid& DstUVGrid, const FMeshDescription& MeshDescription,
+		ProxyLOD::FLinearColorGrid& OutDstBufferGrid)
 	{
+		TVertexAttributesConstRef<FVector> VertexPositions = MeshDescription.VertexAttributes().GetAttributesRef<FVector>(MeshAttribute::Vertex::Position);
+		TVertexInstanceAttributesConstRef<FVector> VertexInstanceNormals = MeshDescription.VertexInstanceAttributes().GetAttributesRef<FVector>(MeshAttribute::VertexInstance::Normal);
+		TVertexInstanceAttributesConstRef<FVector> VertexInstanceTangents = MeshDescription.VertexInstanceAttributes().GetAttributesRef<FVector>(MeshAttribute::VertexInstance::Tangent);
+		TVertexInstanceAttributesConstRef<float> VertexInstanceBinormalSigns = MeshDescription.VertexInstanceAttributes().GetAttributesRef<float>(MeshAttribute::VertexInstance::BinormalSign);
 
 		const FIntPoint SuperSampleSize = SuperSampledNormalGrid.Size();
 		const FIntPoint DstBufferSize = OutDstBufferGrid.Size();
@@ -1263,7 +1271,7 @@ namespace
 
 		const int32 NumXSuperSamples = SuperSampleSize.X / UVSize.X;
 		const int32 NumYSuperSamples = SuperSampleSize.Y / UVSize.Y;
-		
+
 
 		checkSlow(SuperSampleSize.X == UVSize.X * NumXSuperSamples);
 		checkSlow(SuperSampleSize.Y == UVSize.Y * NumYSuperSamples);
@@ -1279,38 +1287,38 @@ namespace
 			typedef openvdb::Mat3R  FTangentSpace;
 
 
-			auto ComputeDstTangentSpace = [&DstRawMesh](const int32 TriangleId, const ProxyLOD::DArray3d& BarycentricCoords, FTangentSpace& TangentSpace)
+			auto ComputeDstTangentSpace = [&MeshDescription, &VertexInstanceTangents, &VertexInstanceBinormalSigns, &VertexInstanceNormals](const int32 TriangleId, const ProxyLOD::DArray3d& BarycentricCoords, FTangentSpace& TangentSpace)
 			{
-				const int32 Indexes[3] = { TriangleId * 3, TriangleId * 3 + 1, TriangleId * 3 + 2 };
+				const FVertexInstanceID VertexInstanceIDs[3] = { FVertexInstanceID(3 * TriangleId),
+																 FVertexInstanceID(3 * TriangleId + 1),
+																 FVertexInstanceID(3 * TriangleId + 2) };
+
 
 				// Compute the local tangent space
 				FVector InterpolatedTangent;
-				//{
-					FVector TangentSamples[3];
-					TangentSamples[0] = DstRawMesh.WedgeTangentX[Indexes[0]];
-					TangentSamples[1] = DstRawMesh.WedgeTangentX[Indexes[1]];
-					TangentSamples[2] = DstRawMesh.WedgeTangentX[Indexes[2]];
-					InterpolatedTangent = ProxyLOD::InterpolateVertexData(BarycentricCoords, TangentSamples);
-				//}
+				FVector TangentSamples[3];
+				TangentSamples[0] = VertexInstanceTangents[VertexInstanceIDs[0]];
+				TangentSamples[1] = VertexInstanceTangents[VertexInstanceIDs[1]];
+				TangentSamples[2] = VertexInstanceTangents[VertexInstanceIDs[2]];
+				InterpolatedTangent = ProxyLOD::InterpolateVertexData(BarycentricCoords, TangentSamples);
 
 				FVector InterpolatedBiTangent;
-
 				// The bi-tangent is reconstructed from the tangent and normal
 				//{
-					FVector BiTangentSamples[3];
-					BiTangentSamples[0] = DstRawMesh.WedgeTangentY[Indexes[0]];
-					BiTangentSamples[1] = DstRawMesh.WedgeTangentY[Indexes[1]];
-					BiTangentSamples[2] = DstRawMesh.WedgeTangentY[Indexes[2]];
-					InterpolatedBiTangent = ProxyLOD::InterpolateVertexData(BarycentricCoords, BiTangentSamples);
+				FVector BiTangentSamples[3];
+				BiTangentSamples[0] = FVector::CrossProduct(VertexInstanceNormals[VertexInstanceIDs[0]], VertexInstanceTangents[VertexInstanceIDs[0]]).GetSafeNormal() * VertexInstanceBinormalSigns[VertexInstanceIDs[0]];
+				BiTangentSamples[1] = FVector::CrossProduct(VertexInstanceNormals[VertexInstanceIDs[1]], VertexInstanceTangents[VertexInstanceIDs[1]]).GetSafeNormal() * VertexInstanceBinormalSigns[VertexInstanceIDs[1]];
+				BiTangentSamples[2] = FVector::CrossProduct(VertexInstanceNormals[VertexInstanceIDs[2]], VertexInstanceTangents[VertexInstanceIDs[2]]).GetSafeNormal() * VertexInstanceBinormalSigns[VertexInstanceIDs[2]];
+				InterpolatedBiTangent = ProxyLOD::InterpolateVertexData(BarycentricCoords, BiTangentSamples);
 				//}
 
 				FVector InterpolatedNormal;
 				//{
-					FVector NormalSamples[3];
-					NormalSamples[0] = DstRawMesh.WedgeTangentZ[Indexes[0]];
-					NormalSamples[1] = DstRawMesh.WedgeTangentZ[Indexes[1]];
-					NormalSamples[2] = DstRawMesh.WedgeTangentZ[Indexes[2]];
-					InterpolatedNormal = ProxyLOD::InterpolateVertexData(BarycentricCoords, NormalSamples);
+				FVector NormalSamples[3];
+				NormalSamples[0] = VertexInstanceNormals[VertexInstanceIDs[0]];
+				NormalSamples[1] = VertexInstanceNormals[VertexInstanceIDs[1]];
+				NormalSamples[2] = VertexInstanceNormals[VertexInstanceIDs[2]];
+				InterpolatedNormal = ProxyLOD::InterpolateVertexData(BarycentricCoords, NormalSamples);
 				//}
 
 				// Testing handedness. Is the tangent x bitangent pointing (roughly) in the direction of the normal?
@@ -1519,7 +1527,7 @@ namespace
 
 
 	template <EFlattenMaterialProperties PropertyType>
-	void TMapFlattenMaterial(const FRawMesh& DstRawMesh, const FRawMeshArrayAdapter& SrcMeshAdapter,
+	void TMapFlattenMaterial(const FMeshDescription& DstRawMesh, const FMeshDescriptionArrayAdapter& SrcMeshAdapter,
 		const ProxyLOD::FSrcDataGrid&  SuperSampledCorrespondenceGrid,
 		const ProxyLOD::FRasterGrid& SuperSampledDstUVGrid, const ProxyLOD::FRasterGrid& DstUVGrid,
 		const TArray<FFlattenMaterial>& InputMaterials, FFlattenMaterial& OutMaterial, const FColor DefaultColor = FColor::Black)
@@ -1682,8 +1690,8 @@ namespace
  // Map Diffuse, Specular, Metallic, Roughness, Normal, Emissive, Opacity to the correct materials.
  // NB: EFlattenMaterialProperties: SubSurface, OpacityMask and AmbientOcclusion are not included
 void MapFlattenMaterial(const EFlattenMaterialProperties PropertyType, 
-	const FRawMesh& DstRawMesh, 
-	const FRawMeshArrayAdapter& SrcMeshAdapter,
+	const FMeshDescription& DstRawMesh,
+	const FMeshDescriptionArrayAdapter& SrcMeshAdapter,
 	const ProxyLOD::FSrcDataGrid& SuperSampledCorrespondenceGrid, 
 	const ProxyLOD::FRasterGrid& SuperSampledDstUVGrid,
 	const ProxyLOD::FRasterGrid& DstUVGrid, 
@@ -1723,7 +1731,7 @@ void MapFlattenMaterial(const EFlattenMaterialProperties PropertyType,
 }// end anonymous namespace
 
 
-void ProxyLOD::MapFlattenMaterials(const FRawMesh& DstRawMesh, const FRawMeshArrayAdapter& SrcMeshAdapter,
+void ProxyLOD::MapFlattenMaterials(const FMeshDescription& DstRawMesh, const FMeshDescriptionArrayAdapter& SrcMeshAdapter,
 	const ProxyLOD::FSrcDataGrid& SuperSampledCorrespondenceGrid,
 	const ProxyLOD::FRasterGrid& SuperSampledDstUVGrid,
 	const ProxyLOD::FRasterGrid& DstUVGrid,
@@ -1759,7 +1767,7 @@ void ProxyLOD::MapFlattenMaterials(const FRawMesh& DstRawMesh, const FRawMeshArr
 
 // Maps each triangle id to a color.  This is just for testing
 void ProxyLOD::ColorMapFlattedMaterials( const FVertexDataMesh& VertexDataMesh, 
-	                                     const FRawMeshArrayAdapter& MeshAdapter, 
+	                                     const FMeshDescriptionArrayAdapter& MeshAdapter,
 	                                     const ProxyLOD::FRasterGrid& UVGrid, 
 	                                     const TArray<FFlattenMaterial>& InputMaterials, 
 	                                     FFlattenMaterial& OutMaterial)
diff --git a/Engine/Plugins/Experimental/ProxyLODPlugin/Source/ProxyLOD/Private/ProxyLODMaterialTransferUtilities.h b/Engine/Plugins/Experimental/ProxyLODPlugin/Source/ProxyLOD/Private/ProxyLODMaterialTransferUtilities.h
--- a/Engine/Plugins/Experimental/ProxyLODPlugin/Source/ProxyLOD/Private/ProxyLODMaterialTransferUtilities.h
+++ b/Engine/Plugins/Experimental/ProxyLODPlugin/Source/ProxyLOD/Private/ProxyLODMaterialTransferUtilities.h
@@ -4,7 +4,7 @@
 
 #include "MaterialUtilities.h"
 
-#include "ProxyLODMeshTypes.h"   // FRawMeshArrayAdapter, FVertexDataMesh
+#include "ProxyLODMeshTypes.h"   // FMeshDescriptionArrayAdapter, FVertexDataMesh
 #include "ProxyLODGrid2d.h"
 #include "ProxyLODRasterizer.h" // FRasterGrid
 
@@ -45,7 +45,7 @@ namespace ProxyLOD
 	* Because some locations in the texture atlas for the Reduced Mesh will fall outside of all the charts (i.e. be dead space),
 	* there will be texels in the SrcDataGrid that correspond to this space and will be given .MaterialId =-1 for quick identification.
 	*/
-	FSrcDataGrid::Ptr CreateCorrespondence( const FRawMeshArrayAdapter& SrcMesh, 
+	FSrcDataGrid::Ptr CreateCorrespondence( const FMeshDescriptionArrayAdapter& SrcMesh,
 										    const FVertexDataMesh& ReducedMesh,
 		                                    const ProxyLOD::FRasterGrid& ReducedMeshUVGrid,
 		                                    const int32 TransferType, 
@@ -75,7 +75,7 @@ namespace ProxyLOD
 	* there will be texels in the SrcDataGrid that correspond to this space and will be given .MaterialId =-1 for quick identification.
 	*/
 	FSrcDataGrid::Ptr CreateCorrespondence( const openvdb::Int32Grid& ClosestSrcPolyGrid,
-		                                    const FRawMeshArrayAdapter& SrcMesh,
+		                                    const FMeshDescriptionArrayAdapter& SrcMesh,
 										    const FVertexDataMesh& ReducedMesh, 
 											const ProxyLOD::FRasterGrid& ReducedMeshUVGrid,
 											const int32 TransferType, 
@@ -100,8 +100,8 @@ namespace ProxyLOD
 	* Map Diffuse, Specular, Metallic, Roughness, Normal, Emissive, Opacity to the correct materials.
 	* NB: EFlattenMaterialProperties: SubSurface, OpacityMask and AmbientOcclusion are not included
 	*/
-	void MapFlattenMaterials( const FRawMesh& DstRawMesh, 
-		                      const FRawMeshArrayAdapter&   SrcMeshAdapter,
+	void MapFlattenMaterials( const FMeshDescription& DstRawMesh,
+		                      const FMeshDescriptionArrayAdapter&   SrcMeshAdapter,
 							  const ProxyLOD::FSrcDataGrid& SuperSampledCorrespondenceGrid,
 							  const ProxyLOD::FRasterGrid&  SuperSampledDstUVGrid,
 							  const ProxyLOD::FRasterGrid&  DstUVGrid,
@@ -115,7 +115,7 @@ namespace ProxyLOD
 	* mapping each triangle id to a color.
 	*/
 	void ColorMapFlattedMaterials( const FVertexDataMesh& VertexDataMesh, 
-		                           const FRawMeshArrayAdapter& MeshAdapter, 
+		                           const FMeshDescriptionArrayAdapter& MeshAdapter,
 		                           const ProxyLOD::FRasterGrid& UVGrid, 
 		                           const TArray<FFlattenMaterial>& InputMaterials, 
 		                           FFlattenMaterial& OutMaterial);
diff --git a/Engine/Plugins/Experimental/ProxyLODPlugin/Source/ProxyLOD/Private/ProxyLODMeshAttrTransfer.cpp b/Engine/Plugins/Experimental/ProxyLODPlugin/Source/ProxyLOD/Private/ProxyLODMeshAttrTransfer.cpp
--- a/Engine/Plugins/Experimental/ProxyLODPlugin/Source/ProxyLOD/Private/ProxyLODMeshAttrTransfer.cpp
+++ b/Engine/Plugins/Experimental/ProxyLODPlugin/Source/ProxyLOD/Private/ProxyLODMeshAttrTransfer.cpp
@@ -58,56 +58,76 @@ FVector2D AverageTexCoord(const FVector2D(&TexCoords)[Size])
 	return Result;
 }
 
-void ProxyLOD::TransferMeshAttributes(const FClosestPolyField& SrcPolyField, FRawMesh& InOutMesh)
+void ProxyLOD::TransferMeshAttributes(const FClosestPolyField& SrcPolyField, FMeshDescription& InOutMesh)
 {
-	const int32 NumFaces = InOutMesh.WedgeIndices.Num() / 3;
-	//const FRawMeshArrayAdapter& RawMeshArrayAdapter = SrcPolyField.MeshAdapter();
+	const int32 NumFaces = InOutMesh.Polygons().Num();
+
+	TVertexAttributesRef<FVector> VertexPositions = InOutMesh.VertexAttributes().GetAttributesRef<FVector>(MeshAttribute::Vertex::Position);
+	TPolygonGroupAttributesRef<FName> PolygonGroupImportedMaterialSlotNames = InOutMesh.PolygonGroupAttributes().GetAttributesRef<FName>(MeshAttribute::PolygonGroup::ImportedMaterialSlotName);
+	TVertexInstanceAttributesRef<FVector> VertexInstanceNormals = InOutMesh.VertexInstanceAttributes().GetAttributesRef<FVector>(MeshAttribute::VertexInstance::Normal);
+	TVertexInstanceAttributesRef<FVector> VertexInstanceTangents = InOutMesh.VertexInstanceAttributes().GetAttributesRef<FVector>(MeshAttribute::VertexInstance::Tangent);
+	TVertexInstanceAttributesRef<float> VertexInstanceBinormalSigns = InOutMesh.VertexInstanceAttributes().GetAttributesRef<float>(MeshAttribute::VertexInstance::BinormalSign);
+	TVertexInstanceAttributesRef<FVector4> VertexInstanceColors = InOutMesh.VertexInstanceAttributes().GetAttributesRef<FVector4>(MeshAttribute::VertexInstance::Color);
+	TVertexInstanceAttributesRef<FVector2D> VertexInstanceUVs = InOutMesh.VertexInstanceAttributes().GetAttributesRef<FVector2D>(MeshAttribute::VertexInstance::TextureCoordinate);
+
+	//const FMeshDescriptionArrayAdapter& RawMeshArrayAdapter = SrcPolyField.MeshAdapter();
 	ProxyLOD::Parallel_For(ProxyLOD::FIntRange(0, NumFaces),
-		[&InOutMesh, &SrcPolyField](const ProxyLOD::FIntRange& Range)
+		[&InOutMesh, &SrcPolyField, &VertexPositions, &VertexInstanceTangents, &VertexInstanceBinormalSigns, &VertexInstanceNormals, &VertexInstanceUVs, &VertexInstanceColors, &PolygonGroupImportedMaterialSlotNames](const ProxyLOD::FIntRange& Range)
 	{
-		auto& Pos = InOutMesh.VertexPositions;
-		auto& WedgeToPos = InOutMesh.WedgeIndices;
 		FClosestPolyField::FPolyConstAccessor ConstPolyAccessor = SrcPolyField.GetPolyConstAccessor();
 
-		uint32  WIdxs[3];
+		FVertexInstanceID WIdxs[3];
 		// loop over the faces
-		for (int32 f = Range.begin(), F = Range.end(); f < F; ++f)
+		for (int32 CurrentRange = Range.begin(), EndRange = Range.end(); CurrentRange < EndRange; ++CurrentRange)
 		{
-			// get the three corners for this face
-			WIdxs[0] = f * 3;
-			WIdxs[1] = WIdxs[0] + 1;
-			WIdxs[2] = WIdxs[0] + 2;
-
+			FPolygonID PolygonID = FPolygonID(CurrentRange);
+			const FMeshPolygon& Polygon = InOutMesh.GetPolygon(PolygonID);
 			int32 LastMaterialIndex = -1;
-			for (int32 i = 0; i < 3; ++i)
+			for (const FMeshTriangle& Triangle : Polygon.Triangles)
 			{
-				const uint32 Idx = WIdxs[i];
-				// world space location 
-				FVector WSPos = Pos[WedgeToPos[Idx]];
+				// get the three corners for this Triangle
+				WIdxs[0] = Triangle.VertexInstanceID0;
+				WIdxs[1] = Triangle.VertexInstanceID1;
+				WIdxs[2] = Triangle.VertexInstanceID2;
 
-				bool bFoundPoly;
-				// The closest poly to this point
-				FRawMeshArrayAdapter::FRawPoly RawPoly = ConstPolyAccessor.Get(WSPos, bFoundPoly);
+				for (int32 i = 0; i < 3; ++i)
+				{
+					const FVertexInstanceID& Idx = WIdxs[i];
+					// world space location
+					const FVector& WSPos = VertexPositions[InOutMesh.GetVertexInstanceVertex(Idx)];
 
-				LastMaterialIndex = RawPoly.FaceMaterialIndex;
+					bool bFoundPoly;
+					// The closest poly to this point
+					FMeshDescriptionArrayAdapter::FRawPoly RawPoly = ConstPolyAccessor.Get(WSPos, bFoundPoly);
 
-				// ----  Transfer the face - average values to each wedge --- //
-				// NB: might replace with something more sophisticated later.
+					LastMaterialIndex = RawPoly.FaceMaterialIndex;
 
-				// Compute the average color
+					// ----  Transfer the face - average values to each wedge --- //
+					// NB: might replace with something more sophisticated later.
 
-				InOutMesh.WedgeColors[Idx] = AverageColor(RawPoly.WedgeColors);
+					// Compute the average color
+					VertexInstanceColors[Idx] = FVector4(FLinearColor(AverageColor(RawPoly.WedgeColors)));
 
-				// The average Tangent Vectors
-				InOutMesh.WedgeTangentX[Idx] = AverageUnitVector(RawPoly.WedgeTangentX);
-				InOutMesh.WedgeTangentY[Idx] = AverageUnitVector(RawPoly.WedgeTangentY);
-				InOutMesh.WedgeTangentZ[Idx] = AverageUnitVector(RawPoly.WedgeTangentZ);
+					// The average Tangent Vectors
+					VertexInstanceTangents[Idx] = AverageUnitVector(RawPoly.WedgeTangentX);
+					VertexInstanceNormals[Idx] = AverageUnitVector(RawPoly.WedgeTangentZ);
+					VertexInstanceBinormalSigns[Idx] = GetBasisDeterminantSign(VertexInstanceTangents[Idx], AverageUnitVector(RawPoly.WedgeTangentY), VertexInstanceNormals[Idx]);
 
-				// Average Texture Coords
-				InOutMesh.WedgeTexCoords[0][Idx] = AverageTexCoord(RawPoly.WedgeTexCoords[0]);
+					// Average Texture Coords
+					VertexInstanceUVs.Set(Idx, 0, AverageTexCoord(RawPoly.WedgeTexCoords[0]));
+				}
 			}
 			// Assign the material index that the last vertex of this face sees.
-			InOutMesh.FaceMaterialIndices[f] = LastMaterialIndex;
+			if (LastMaterialIndex != -1)
+			{
+				FPolygonGroupID PolygonGroupID(LastMaterialIndex);
+				if (!InOutMesh.IsPolygonGroupValid(PolygonGroupID))
+				{
+					InOutMesh.CreatePolygonGroupWithID(PolygonGroupID);
+					PolygonGroupImportedMaterialSlotNames[PolygonGroupID] = FName(*FString::Printf(TEXT("ProxyLOD_Material_%d"), FMath::Rand()));
+				}
+				InOutMesh.SetPolygonPolygonGroup(PolygonID, PolygonGroupID);
+			}
 		}
 	}
 	);
@@ -125,14 +145,14 @@ void ProxyLOD::TransferSrcNormals(const FClosestPolyField& SrcPolyField, FAOSMes
 		[&Vertexes, &SrcPolyField](const ProxyLOD::FUIntRange& Range)
 	{
 		const auto PolyAccessor = SrcPolyField.GetPolyConstAccessor();
-		for (uint32 n = Range.begin(), N = Range.end(); n < N; ++n)
+		for (uint32 CurrentRange = Range.begin(), EndRange = Range.end(); CurrentRange < EndRange; ++CurrentRange)
 		{
 			// Get the closest poly to this vertex.
-			FVector& Normal = Vertexes[n].Normal;
-			const FVector& Pos = Vertexes[n].Position;
+			FVector& Normal = Vertexes[CurrentRange].Normal;
+			const FVector& Pos = Vertexes[CurrentRange].Position;
 
 			bool bSuccess = false;
-			const FRawMeshArrayAdapter::FRawPoly  RawPoly = PolyAccessor.Get(openvdb::Vec3d(Pos.X, Pos.Y, Pos.Z), bSuccess);
+			const FMeshDescriptionArrayAdapter::FRawPoly  RawPoly = PolyAccessor.Get(openvdb::Vec3d(Pos.X, Pos.Y, Pos.Z), bSuccess);
 			if (bSuccess)
 			{
 				bool bValidSizedPoly = true;
@@ -172,39 +192,36 @@ void ProxyLOD::TransferSrcNormals(const FClosestPolyField& SrcPolyField, FAOSMes
 };
 
 
-void ProxyLOD::TransferVertexColors(const FClosestPolyField& SrcPolyField, FRawMesh& InOutMesh)
+void ProxyLOD::TransferVertexColors(const FClosestPolyField& SrcPolyField, FMeshDescription& InOutMesh)
 {
-	const uint32 NumVertexes = InOutMesh.VertexPositions.Num();
-	FVector* Vertexes = InOutMesh.VertexPositions.GetData();
+	TVertexAttributesConstRef<FVector> VertexPositions = InOutMesh.VertexAttributes().GetAttributesRef<FVector>(MeshAttribute::Vertex::Position);
+	TVertexInstanceAttributesRef<FVector4> VertexInstanceColors = InOutMesh.VertexInstanceAttributes().GetAttributesRef<FVector4>(MeshAttribute::VertexInstance::Color);
 
-	const uint32 NumWedges = InOutMesh.WedgeIndices.Num();
-	uint32* Indices = InOutMesh.WedgeIndices.GetData();
+	const uint32 NumWedges = InOutMesh.VertexInstances().Num();
 
-	FColor* Colors = InOutMesh.WedgeColors.GetData();
 
 	// Loop over the polys in the result mesh.
 
 	ProxyLOD::Parallel_For(ProxyLOD::FUIntRange(0, NumWedges),
-		[Vertexes, Indices, Colors, &SrcPolyField](const ProxyLOD::FUIntRange& Range)
+		[&VertexPositions, &VertexInstanceColors, &InOutMesh, &SrcPolyField](const ProxyLOD::FUIntRange& Range)
 	{
 		const auto PolyAccessor = SrcPolyField.GetPolyConstAccessor();
 
 		// loop over wedges.
-		for (uint32 w = Range.begin(), W = Range.end(); w < W; ++w)
+		for (uint32 CurrentRange = Range.begin(), EndRange = Range.end(); CurrentRange < EndRange; ++CurrentRange)
 		{
 			// Get the closest poly to this vertex.
-			
-			uint32 Idx = Indices[w];
-			const FVector& Pos = Vertexes[Idx];
+			FVertexInstanceID VertexInstanceID(CurrentRange);
+			const FVector& Pos = VertexPositions[InOutMesh.GetVertexInstanceVertex(VertexInstanceID)];
 
 			// Find the closest poly to this wedge.  
 			// NB: all wedges that share a vert location will end up with the same color this way
 
 			bool bSuccess = false;
-			const FRawMeshArrayAdapter::FRawPoly  RawPoly = PolyAccessor.Get(openvdb::Vec3d(Pos.X, Pos.Y, Pos.Z), bSuccess);
+			const FMeshDescriptionArrayAdapter::FRawPoly  RawPoly = PolyAccessor.Get(openvdb::Vec3d(Pos.X, Pos.Y, Pos.Z), bSuccess);
 			
 			// default to white
-			Colors[w] = FColor::White;
+			VertexInstanceColors[VertexInstanceID] = FVector4(FLinearColor::White);
 			if (bSuccess)
 			{
 
@@ -234,7 +251,7 @@ void ProxyLOD::TransferVertexColors(const FClosestPolyField& SrcPolyField, FRawM
 
 					// fix up the intensity.
 
-					Colors[w] = InterpolatedColor.ToFColor(true);
+					VertexInstanceColors[VertexInstanceID] = FVector4(InterpolatedColor);
 				}
 			}
 			
@@ -244,23 +261,24 @@ void ProxyLOD::TransferVertexColors(const FClosestPolyField& SrcPolyField, FRawM
 
 
 template <typename ProjectionOperatorType>
-void ProjectVerticiesOntoSrc(const ProjectionOperatorType& ProjectionOperator, const FClosestPolyField& SrcPolyField, FRawMesh& InOutMesh)
+void ProjectVerticiesOntoSrc(const ProjectionOperatorType& ProjectionOperator, const FClosestPolyField& SrcPolyField, FMeshDescription& InOutMesh)
 {
-	const uint32 NumVertexes = InOutMesh.VertexPositions.Num();
-	FVector* Vertexes = InOutMesh.VertexPositions.GetData();
+	TVertexAttributesRef<FVector> VertexPositions = InOutMesh.VertexAttributes().GetAttributesRef<FVector>(MeshAttribute::Vertex::Position);
+	const uint32 NumVertexes = InOutMesh.Vertices().Num();
 
 	ProxyLOD::Parallel_For(ProxyLOD::FUIntRange(0, NumVertexes),
-		[Vertexes, &SrcPolyField, ProjectionOperator](const ProxyLOD::FUIntRange& Range)
+		[VertexPositions, &SrcPolyField, ProjectionOperator](const ProxyLOD::FUIntRange& Range)
 	{
 		const auto PolyAccessor = SrcPolyField.GetPolyConstAccessor();
 
-		for (uint32 n = Range.begin(), N = Range.end(); n < N; ++n)
+		for (uint32 CurrentRange = Range.begin(), EndRange = Range.end(); CurrentRange < EndRange; ++CurrentRange)
 		{
+			FVertexID VertexID(CurrentRange);
 			// Get the closest poly to this vertex.
-			FVector& Pos = Vertexes[n];
+			FVector& Pos = VertexPositions[VertexID];
 
 			bool bSuccess = false;
-			const FRawMeshArrayAdapter::FRawPoly  RawPoly = PolyAccessor.Get(openvdb::Vec3d(Pos.X, Pos.Y, Pos.Z), bSuccess);
+			const FMeshDescriptionArrayAdapter::FRawPoly  RawPoly = PolyAccessor.Get(openvdb::Vec3d(Pos.X, Pos.Y, Pos.Z), bSuccess);
 			if (bSuccess)
 			{
 
@@ -292,13 +310,13 @@ void ProjectVerticiesOntoSrc(const ProjectionOperatorType& ProjectionOperator, c
 	{
 		const auto PolyAccessor = SrcPolyField.GetPolyConstAccessor();
 
-		for (uint32 n = Range.begin(), N = Range.end(); n < N; ++n)
+		for (uint32 CurrentRange = Range.begin(), EndRange = Range.end(); CurrentRange < EndRange; ++CurrentRange)
 		{
 			// Get the closest poly to this vertex.
-			FVector& Pos = Vertexes[n];
+			FVector& Pos = Vertexes[CurrentRange];
 
 			bool bSuccess = false;
-			const FRawMeshArrayAdapter::FRawPoly  RawPoly = PolyAccessor.Get(openvdb::Vec3d(Pos.X, Pos.Y, Pos.Z), bSuccess);
+			const FMeshDescriptionArrayAdapter::FRawPoly  RawPoly = PolyAccessor.Get(openvdb::Vec3d(Pos.X, Pos.Y, Pos.Z), bSuccess);
 			if (bSuccess)
 			{
 
@@ -330,13 +348,13 @@ void ProjectVerticiesOntoSrc(const ProjectionOperatorType& ProjectionOperator, c
 	{
 		const auto PolyAccessor = SrcPolyField.GetPolyConstAccessor();
 
-		for (uint32 n = Range.begin(), N = Range.end(); n < N; ++n)
+		for (uint32 CurrentRange = Range.begin(), EndRange = Range.end(); CurrentRange < EndRange; ++CurrentRange)
 		{
 			// Get the closest poly to this vertex.
-			FVector& Pos = Vertexes[n].Position;
+			FVector& Pos = Vertexes[CurrentRange].Position;
 
 			bool bSuccess = false;
-			const FRawMeshArrayAdapter::FRawPoly  RawPoly = PolyAccessor.Get(openvdb::Vec3d(Pos.X, Pos.Y, Pos.Z), bSuccess);
+			const FMeshDescriptionArrayAdapter::FRawPoly  RawPoly = PolyAccessor.Get(openvdb::Vec3d(Pos.X, Pos.Y, Pos.Z), bSuccess);
 			if (bSuccess)
 			{
 
@@ -395,7 +413,7 @@ private:
 
 };
 
-void ProxyLOD::ProjectVertexWithSnapToNearest(const FClosestPolyField& SrcPolyField, FRawMesh& InOutMesh)
+void ProxyLOD::ProjectVertexWithSnapToNearest(const FClosestPolyField& SrcPolyField, FMeshDescription& InOutMesh)
 {
 
 	const double VoxelSize = SrcPolyField.GetVoxelSize();
@@ -472,7 +490,7 @@ private:
 
 };
 
-void ProxyLOD::ProjectVertexOntoSrcSurface(const FClosestPolyField& SrcPolyField, FRawMesh& InOutMesh)
+void ProxyLOD::ProjectVertexOntoSrcSurface(const FClosestPolyField& SrcPolyField, FMeshDescription& InOutMesh)
 {
 
 	const double VoxelSize = SrcPolyField.GetVoxelSize();
diff --git a/Engine/Plugins/Experimental/ProxyLODPlugin/Source/ProxyLOD/Private/ProxyLODMeshAttrTransfer.h b/Engine/Plugins/Experimental/ProxyLODPlugin/Source/ProxyLOD/Private/ProxyLODMeshAttrTransfer.h
--- a/Engine/Plugins/Experimental/ProxyLODPlugin/Source/ProxyLOD/Private/ProxyLODMeshAttrTransfer.h
+++ b/Engine/Plugins/Experimental/ProxyLODPlugin/Source/ProxyLOD/Private/ProxyLODMeshAttrTransfer.h
@@ -19,7 +19,7 @@ namespace ProxyLOD
 	* @param  SrcPolyField  The closest poly field derived from the high poly src mesh
 	* @param  InOutMesh     The target raw mesh.  Only the wedge colors on this target are updated.
 	*/
-	void TransferVertexColors(const FClosestPolyField& SrcPolyField, FRawMesh& InOutMesh);
+	void TransferVertexColors(const FClosestPolyField& SrcPolyField, FMeshDescription& InOutMesh);
 
 	
 	/**
@@ -59,7 +59,7 @@ namespace ProxyLOD
 	* @param  SrcPolyField  The closest poly field derived from the high poly src mesh
 	* @param  InOutMesh     The target low poly mesh.  Only the position is updated.
 	*/
-	void ProjectVertexWithSnapToNearest( const FClosestPolyField& SrcPolyField, FRawMesh& InOutMesh );
+	void ProjectVertexWithSnapToNearest( const FClosestPolyField& SrcPolyField, FMeshDescription& InOutMesh );
 	void ProjectVertexWithSnapToNearest( const FClosestPolyField& SrcPolyField, FAOSMesh& InOutMesh );
 	void ProjectVertexWithSnapToNearest( const FClosestPolyField& SrcPolyField, FVertexDataMesh& InOutMesh);
 
@@ -75,7 +75,7 @@ namespace ProxyLOD
 	* @param  SrcPolyField  The closest poly field derived from the high poly src mesh
 	* @param  InOutMesh     The target low poly mesh.  Only the position is updated.
 	*/
-	void ProjectVertexOntoSrcSurface( const FClosestPolyField& SrcPolyField, FRawMesh& InOutMesh );
+	void ProjectVertexOntoSrcSurface( const FClosestPolyField& SrcPolyField, FMeshDescription& InOutMesh );
 	void ProjectVertexOntoSrcSurface( const FClosestPolyField& SrcPolyField, FAOSMesh& InOutMesh );
 	void ProjectVertexOntoSrcSurface( const FClosestPolyField& SrcPolyField, FVertexDataMesh& InOutMesh );
 
@@ -92,7 +92,7 @@ namespace ProxyLOD
 	* @param  SrcPolyField  The closest poly field derived from the high poly src mesh
 	* @param  InOutMesh     The target low poly mesh.
 	*/
-	void TransferMeshAttributes(const FClosestPolyField& SrcPolyField, FRawMesh& InOutMesh);
+	void TransferMeshAttributes(const FClosestPolyField& SrcPolyField, FMeshDescription& InOutMesh);
 
 }
 
diff --git a/Engine/Plugins/Experimental/ProxyLODPlugin/Source/ProxyLOD/Private/ProxyLODMeshConvertUtils.cpp b/Engine/Plugins/Experimental/ProxyLODPlugin/Source/ProxyLOD/Private/ProxyLODMeshConvertUtils.cpp
--- a/Engine/Plugins/Experimental/ProxyLODPlugin/Source/ProxyLOD/Private/ProxyLODMeshConvertUtils.cpp
+++ b/Engine/Plugins/Experimental/ProxyLODPlugin/Source/ProxyLOD/Private/ProxyLODMeshConvertUtils.cpp
@@ -3,9 +3,21 @@
 #include "ProxyLODMeshConvertUtils.h" 
 #include "ProxyLODMeshUtilities.h"
 
+#include "MeshDescriptionOperations.h"
+
 // Convert QuadMesh to Triangles by splitting
-void ProxyLOD::MixedPolyMeshToRawMesh(const FMixedPolyMesh& SimpleMesh, FRawMesh& DstRawMesh)
+void ProxyLOD::MixedPolyMeshToRawMesh(const FMixedPolyMesh& SimpleMesh, FMeshDescription& DstRawMesh)
 {
+	TVertexAttributesRef<FVector> VertexPositions = DstRawMesh.VertexAttributes().GetAttributesRef<FVector>(MeshAttribute::Vertex::Position);
+	TEdgeAttributesRef<bool> EdgeHardnesses = DstRawMesh.EdgeAttributes().GetAttributesRef<bool>(MeshAttribute::Edge::IsHard);
+	TEdgeAttributesRef<float> EdgeCreaseSharpnesses = DstRawMesh.EdgeAttributes().GetAttributesRef<float>(MeshAttribute::Edge::CreaseSharpness);
+	TPolygonGroupAttributesRef<FName> PolygonGroupImportedMaterialSlotNames = DstRawMesh.PolygonGroupAttributes().GetAttributesRef<FName>(MeshAttribute::PolygonGroup::ImportedMaterialSlotName);
+	TVertexInstanceAttributesRef<FVector> VertexInstanceNormals = DstRawMesh.VertexInstanceAttributes().GetAttributesRef<FVector>(MeshAttribute::VertexInstance::Normal);
+	TVertexInstanceAttributesRef<FVector> VertexInstanceTangents = DstRawMesh.VertexInstanceAttributes().GetAttributesRef<FVector>(MeshAttribute::VertexInstance::Tangent);
+	TVertexInstanceAttributesRef<float> VertexInstanceBinormalSigns = DstRawMesh.VertexInstanceAttributes().GetAttributesRef<float>(MeshAttribute::VertexInstance::BinormalSign);
+	TVertexInstanceAttributesRef<FVector4> VertexInstanceColors = DstRawMesh.VertexInstanceAttributes().GetAttributesRef<FVector4>(MeshAttribute::VertexInstance::Color);
+	TVertexInstanceAttributesRef<FVector2D> VertexInstanceUVs = DstRawMesh.VertexInstanceAttributes().GetAttributesRef<FVector2D>(MeshAttribute::VertexInstance::TextureCoordinate);
+
 	// Splitting a quad doesn't introduce any new verts.
 	const uint32 DstNumVerts = SimpleMesh.Points.size();
 
@@ -16,469 +28,410 @@ void ProxyLOD::MixedPolyMeshToRawMesh(const FMixedPolyMesh& SimpleMesh, FRawMesh
 
 	// Each Triangle has 3 corners
 	const uint32 DstNumIndexes = 3 * DstNumTris;
-	// Copy the vertices over
+
+	if (VertexInstanceUVs.GetNumIndices() < 1)
 	{
-		// Allocate the space for the verts in the DstRawMesh
+		VertexInstanceUVs.SetNumIndices(1);
+	}
 
-		ResizeArray(DstRawMesh.VertexPositions, DstNumVerts);
+	FPolygonGroupID PolygonGroupID = FPolygonGroupID::Invalid;
+	if (DstRawMesh.PolygonGroups().Num() == 0)
+	{
+		PolygonGroupID = DstRawMesh.CreatePolygonGroup();
+		PolygonGroupImportedMaterialSlotNames[PolygonGroupID] = FName(*FString::Printf(TEXT("ProxyLOD_Material_%d"), FMath::Rand()));
+	}
+	else
+	{
+		PolygonGroupID = DstRawMesh.PolygonGroups().GetFirstValidID();
+	}
 
-		ProxyLOD::Parallel_For(ProxyLOD::FUIntRange(0, DstNumVerts),
-			[&SimpleMesh, &DstRawMesh](const ProxyLOD::FUIntRange& Range)
+	// Copy the vertices over
+	TMap<int32, FVertexID> VertexIDMap;
+	VertexIDMap.Reserve(DstNumVerts);
+	{
+		for (uint32 i = 0, I = DstNumVerts; i < I; ++i)
 		{
-	
-			for (uint32 i = Range.begin(), I = Range.end(); i < I; ++i)
-			{
-				const openvdb::Vec3s& Vertex = SimpleMesh.Points[i];
-				DstRawMesh.VertexPositions[i] = FVector(Vertex[0], Vertex[1], Vertex[2]);
-			}
+			const openvdb::Vec3s& Vertex = SimpleMesh.Points[i];
+			const FVertexID NewVertexID = DstRawMesh.CreateVertex();
+			VertexPositions[NewVertexID] = FVector(Vertex[0], Vertex[1], Vertex[2]);
+			VertexIDMap.Add(i, NewVertexID);
 		}
-		);
 	}
 
-	// Connectivity: 
+	// Connectivity:
+	auto CreateTriangle = [&DstRawMesh, PolygonGroupID, &VertexInstanceNormals, &VertexInstanceTangents, &VertexInstanceBinormalSigns, &VertexInstanceColors, &VertexInstanceUVs, &EdgeHardnesses, &EdgeCreaseSharpnesses](FVertexID TriangleIndex[3])
 	{
-		// convert each quad into two triangles.
-		TArray<uint32>& WedgeIndices = DstRawMesh.WedgeIndices;
-		ResizeArray(WedgeIndices, DstNumIndexes);
-
-		ProxyLOD::Parallel_For(ProxyLOD::FUIntRange(0, NumQuads),
-			[&WedgeIndices, &SimpleMesh](const ProxyLOD::FUIntRange& Range)
+		FVertexInstanceID VertexInstanceIDs[3];
+		for (int32 Corner = 0; Corner < 3; ++Corner)
 		{
-			for (uint32 q = Range.begin(), Q = Range.end(); q < Q; ++q)
-			{
-				const uint32 Offset = q * 6;
-				const openvdb::Vec4I& Quad = SimpleMesh.Quads[q];
-				// add as two triangles to raw mesh
-#if (PROXYLOD_CLOCKWISE_TRIANGLES == 1)
-				// first triangle
-				WedgeIndices[Offset] = Quad[0];
-				WedgeIndices[Offset + 1] = Quad[1];
-				WedgeIndices[Offset + 2] = Quad[2];
-				// second triangle
-				WedgeIndices[Offset + 3] = Quad[2];
-				WedgeIndices[Offset + 4] = Quad[3];
-				WedgeIndices[Offset + 5] = Quad[0];
-#else
-				WedgeIndices[Offset] = Quad[0];
-				WedgeIndices[Offset + 1] = Quad[3];
-				WedgeIndices[Offset + 2] = Quad[2];
-				// second triangle
-				WedgeIndices[Offset + 3] = Quad[2];
-				WedgeIndices[Offset + 4] = Quad[1];
-				WedgeIndices[Offset + 5] = Quad[0];
-#endif 
-			}
-
-		});
+			VertexInstanceIDs[Corner] = DstRawMesh.CreateVertexInstance(TriangleIndex[Corner]);
+			VertexInstanceTangents[VertexInstanceIDs[Corner]] = FVector(1, 0, 0);
+			VertexInstanceNormals[VertexInstanceIDs[Corner]] = FVector(0, 0, 1);
+			VertexInstanceBinormalSigns[VertexInstanceIDs[Corner]] = GetBasisDeterminantSign(VertexInstanceTangents[VertexInstanceIDs[Corner]].GetSafeNormal(),
+																							 (VertexInstanceNormals[VertexInstanceIDs[Corner]] ^ VertexInstanceTangents[VertexInstanceIDs[Corner]]).GetSafeNormal(),
+																							 VertexInstanceNormals[VertexInstanceIDs[Corner]].GetSafeNormal());
+			VertexInstanceColors[VertexInstanceIDs[Corner]] = FVector4(0.0f);
+			VertexInstanceUVs.Set(VertexInstanceIDs[Corner], 0, FVector2D(0.0f, 0.0f));
+		}
 
-		// add the SimpleMesh triangles.
-		ProxyLOD::Parallel_For(ProxyLOD::FUIntRange(0, SimpleMesh.Triangles.size()),
-			[&WedgeIndices, &SimpleMesh, NumQuads](const ProxyLOD::FUIntRange& Range)
+		//Create a polygon from this triangle
+		TArray<FMeshDescription::FContourPoint> Contours;
+		for (int32 Corner = 0; Corner < 3; ++Corner)
 		{
-			for (uint32 t = Range.begin(), T = Range.end(); t < T; ++t)
+			int32 ContourPointIndex = Contours.AddDefaulted();
+			FMeshDescription::FContourPoint& ContourPoint = Contours[ContourPointIndex];
+			//Find the matching edge ID
+			uint32 CornerIndices[2];
+			CornerIndices[0] = (Corner + 0) % 3;
+			CornerIndices[1] = (Corner + 1) % 3;
+
+			FVertexID EdgeVertexIDs[2];
+			EdgeVertexIDs[0] = TriangleIndex[CornerIndices[0]];
+			EdgeVertexIDs[1] = TriangleIndex[CornerIndices[1]];
+
+			FEdgeID MatchEdgeId = DstRawMesh.GetVertexPairEdge(EdgeVertexIDs[0], EdgeVertexIDs[1]);
+			if (MatchEdgeId == FEdgeID::Invalid)
 			{
-				const uint32 Offset = NumQuads * 6 + t * 3;
-				const openvdb::Vec3I& Tri = SimpleMesh.Triangles[t];
-				// add the triangle to the raw mesh
-#if (PROXYLOD_CLOCKWISE_TRIANGLES == 1)
-				WedgeIndices[Offset] = Tri[0];
-				WedgeIndices[Offset + 1] = Tri[1];
-				WedgeIndices[Offset + 2] = Tri[2];
-#else
-				WedgeIndices[Offset] = Tri[2];
-				WedgeIndices[Offset + 1] = Tri[1];
-				WedgeIndices[Offset + 2] = Tri[0];
-#endif 
+				MatchEdgeId = DstRawMesh.CreateEdge(EdgeVertexIDs[0], EdgeVertexIDs[1]);
+				EdgeHardnesses[MatchEdgeId] = false;
+				EdgeCreaseSharpnesses[MatchEdgeId] = 0.0f;
 			}
+			ContourPoint.EdgeID = MatchEdgeId;
+			ContourPoint.VertexInstanceID = VertexInstanceIDs[CornerIndices[0]];
 		}
-		);
-	}
-
+		// Insert a polygon into the mesh
+		const FPolygonID NewPolygonID = DstRawMesh.CreatePolygon(PolygonGroupID, Contours);
+		//Triangulate the polygon
+		FMeshPolygon& Polygon = DstRawMesh.GetPolygon(NewPolygonID);
+		DstRawMesh.ComputePolygonTriangulation(NewPolygonID, Polygon.Triangles);
+	};
 
-	ResizeArray(DstRawMesh.WedgeTangentX, DstNumIndexes);
-	ResizeArray(DstRawMesh.WedgeTangentY, DstNumIndexes);
-	ResizeArray(DstRawMesh.WedgeTangentZ, DstNumIndexes);
-
-
-	// Generate Tangent-Space
-	// NB: adding default values for now
-	ProxyLOD::Parallel_For(ProxyLOD::FUIntRange(0, DstNumIndexes), [&DstRawMesh](const ProxyLOD::FUIntRange& Range)
 	{
-		for (uint32 i = Range.begin(), I = Range.end(); i < I; ++i)
+		for (uint32 q = 0, Q = NumQuads; q < Q; ++q)
 		{
-			DstRawMesh.WedgeTangentX[i] = FVector(1, 0, 0);
-			DstRawMesh.WedgeTangentY[i] = FVector(0, 1, 0);
-			// TangentZ is the normal - will be over-written later
-			DstRawMesh.WedgeTangentZ[i] = FVector(0, 0, 1);
+			const openvdb::Vec4I& Quad = SimpleMesh.Quads[q];
+			// add as two triangles to raw mesh
+			bool bClockWiseTriangle = true;
+#if (PROXYLOD_CLOCKWISE_TRIANGLES != 1)
+			bClockWiseTriangle = false;
+#endif
+			FVertexID VertexIndexes[3];
+			VertexIndexes[0] = VertexIDMap[Quad[0]];
+			VertexIndexes[1] = bClockWiseTriangle ? VertexIDMap[Quad[1]] : VertexIDMap[Quad[3]];
+			VertexIndexes[2] = VertexIDMap[Quad[2]];
+			CreateTriangle(VertexIndexes);
+			VertexIndexes[0] = VertexIDMap[Quad[2]];
+			VertexIndexes[1] = bClockWiseTriangle ? VertexIDMap[Quad[3]] : VertexIDMap[Quad[1]];
+			VertexIndexes[2] = VertexIDMap[Quad[0]];
+			CreateTriangle(VertexIndexes);
 		}
-	});
 
-
-
-	DstRawMesh.FaceMaterialIndices.Empty(DstNumTris);
-	DstRawMesh.FaceMaterialIndices.AddZeroed(DstNumTris);
-
-	DstRawMesh.FaceSmoothingMasks.Empty(DstNumTris);
-	DstRawMesh.FaceSmoothingMasks.AddZeroed(DstNumTris);
-
-	DstRawMesh.WedgeColors.Empty();
-	DstRawMesh.WedgeColors.AddZeroed(3 * DstNumTris);
-	for (int32 TexCoordIndex = 0; TexCoordIndex < 8; ++TexCoordIndex)
-	{
-		DstRawMesh.WedgeTexCoords[TexCoordIndex].Empty();
+		// add the SimpleMesh triangles.
+		uint32 IndexStop = SimpleMesh.Triangles.size();
+		for (uint32 t = 0, T = IndexStop; t < T; ++t)
+		{
+			const openvdb::Vec3I& Tri = SimpleMesh.Triangles[t];
+			bool bClockWiseTriangle = true;
+#if (PROXYLOD_CLOCKWISE_TRIANGLES != 1)
+			bClockWiseTriangle = false;
+#endif
+			FVertexID VertexIndexes[3];
+			VertexIndexes[0] = bClockWiseTriangle ? VertexIDMap[Tri[0]] : VertexIDMap[Tri[2]];
+			VertexIndexes[1] = VertexIDMap[Tri[1]];
+			VertexIndexes[2] = bClockWiseTriangle ? VertexIDMap[Tri[2]] : VertexIDMap[Tri[0]];
+			CreateTriangle(VertexIndexes);
+		}
 	}
-	// Need to have a a WedgeTextCoord channel
-	auto& WedgeTexCoord = DstRawMesh.WedgeTexCoords[0];
-	WedgeTexCoord.AddZeroed(DstNumIndexes);
 }
 
 
-void ProxyLOD::AOSMeshToRawMesh(const FAOSMesh& AOSMesh, FRawMesh& OutRawMesh)
+void ProxyLOD::AOSMeshToRawMesh(const FAOSMesh& AOSMesh, FMeshDescription& OutRawMesh)
 {
-
+	TVertexAttributesRef<FVector> VertexPositions = OutRawMesh.VertexAttributes().GetAttributesRef<FVector>(MeshAttribute::Vertex::Position);
+	TEdgeAttributesRef<bool> EdgeHardnesses = OutRawMesh.EdgeAttributes().GetAttributesRef<bool>(MeshAttribute::Edge::IsHard);
+	TEdgeAttributesRef<float> EdgeCreaseSharpnesses = OutRawMesh.EdgeAttributes().GetAttributesRef<float>(MeshAttribute::Edge::CreaseSharpness);
+	TPolygonGroupAttributesRef<FName> PolygonGroupImportedMaterialSlotNames = OutRawMesh.PolygonGroupAttributes().GetAttributesRef<FName>(MeshAttribute::PolygonGroup::ImportedMaterialSlotName);
+	TVertexInstanceAttributesRef<FVector> VertexInstanceNormals = OutRawMesh.VertexInstanceAttributes().GetAttributesRef<FVector>(MeshAttribute::VertexInstance::Normal);
+	TVertexInstanceAttributesRef<FVector> VertexInstanceTangents = OutRawMesh.VertexInstanceAttributes().GetAttributesRef<FVector>(MeshAttribute::VertexInstance::Tangent);
+	TVertexInstanceAttributesRef<float> VertexInstanceBinormalSigns = OutRawMesh.VertexInstanceAttributes().GetAttributesRef<float>(MeshAttribute::VertexInstance::BinormalSign);
+	TVertexInstanceAttributesRef<FVector4> VertexInstanceColors = OutRawMesh.VertexInstanceAttributes().GetAttributesRef<FVector4>(MeshAttribute::VertexInstance::Color);
+	TVertexInstanceAttributesRef<FVector2D> VertexInstanceUVs = OutRawMesh.VertexInstanceAttributes().GetAttributesRef<FVector2D>(MeshAttribute::VertexInstance::TextureCoordinate);
 
 	const uint32 DstNumPositions = AOSMesh.GetNumVertexes();
 	const uint32 DstNumIndexes = AOSMesh.GetNumIndexes();
 
-	checkSlow(DstNumIndexes % 3 == 0);
-	// Copy the vertices over
+	if (VertexInstanceUVs.GetNumIndices() < 1)
 	{
-		// Allocate the space for the verts in the DstRawMesh
-		auto& VertexPositions = OutRawMesh.VertexPositions;
-		ResizeArray(VertexPositions, DstNumPositions);
-
-		ProxyLOD::Parallel_For(ProxyLOD::FUIntRange(0, DstNumPositions),
-			[&AOSMesh, &VertexPositions](const ProxyLOD::FUIntRange& Range)
-		{
-			const auto& AOSVertexes = AOSMesh.Vertexes;
-			for (uint32 i = Range.begin(), I = Range.end(); i < I; ++i)
-			{
-				const FVector& Position = AOSVertexes[i].GetPos();
-				VertexPositions[i] = Position;
-			}
-		});
-
-		checkSlow(OutRawMesh.VertexPositions.Num() == DstNumPositions);
+		VertexInstanceUVs.SetNumIndices(1);
 	}
 
-	// Connectivity: 
+	FPolygonGroupID PolygonGroupID = FPolygonGroupID::Invalid;
+	if (OutRawMesh.PolygonGroups().Num() == 0)
 	{
-		// convert each quad into two triangles.
-		TArray<uint32>& WedgeIndices = OutRawMesh.WedgeIndices;
-
-		ResizeArray(WedgeIndices, DstNumIndexes);
-
-		ProxyLOD::Parallel_For(ProxyLOD::FUIntRange(0, DstNumIndexes),
-			[&AOSMesh, &WedgeIndices](const ProxyLOD::FUIntRange& Range)
-		{
-			const auto& AOSIndexes = AOSMesh.Indexes;
-			for (uint32 i = Range.begin(), I = Range.end(); i < I; ++i)
-			{
-				WedgeIndices[i] = AOSIndexes[i];
-			}
-		});
-
-		checkSlow(WedgeIndices.Num() == DstNumIndexes);
+		PolygonGroupID = OutRawMesh.CreatePolygonGroup();
+		PolygonGroupImportedMaterialSlotNames[PolygonGroupID] = FName(*FString::Printf(TEXT("ProxyLOD_Material_%d"), FMath::Rand()));
 	}
-
-	// Add default Tangents
+	else
 	{
-
-		ResizeArray(OutRawMesh.WedgeTangentX, DstNumIndexes);
-		ResizeArray(OutRawMesh.WedgeTangentY, DstNumIndexes);
-
-		// Generate Tangent-Space
-		// NB: adding default values for now
-		ProxyLOD::Parallel_For(ProxyLOD::FUIntRange(0, DstNumIndexes), [&OutRawMesh](const ProxyLOD::FUIntRange& Range)
-		{
-			for (uint32 i = Range.begin(), I = Range.end(); i < I; ++i)
-			{
-				OutRawMesh.WedgeTangentX[i] = FVector(1, 0, 0);
-				OutRawMesh.WedgeTangentY[i] = FVector(0, 1, 0);
-			}
-		});
+		PolygonGroupID = OutRawMesh.PolygonGroups().GetFirstValidID();
 	}
 
-	// Transfer the normal
+	checkSlow(DstNumIndexes % 3 == 0);
+	// Copy the vertices over
+	TMap<int32, FVertexID> VertexIDMap;
+	VertexIDMap.Reserve(DstNumPositions);
 	{
-		TArray<FVector>& NormalArray = OutRawMesh.WedgeTangentZ;
-		ResizeArray(NormalArray, DstNumIndexes);
-
-		ProxyLOD::Parallel_For(ProxyLOD::FUIntRange(0, DstNumIndexes),
-			[&AOSMesh, &NormalArray](const ProxyLOD::FUIntRange& Range)
+		const auto& AOSVertexes = AOSMesh.Vertexes;
+		for (uint32 i = 0, I = DstNumPositions; i < I; ++i)
 		{
-			const uint32* Indexes = AOSMesh.Indexes;
+			const FVector& Position = AOSVertexes[i].GetPos();
+			const FVertexID NewVertexID = OutRawMesh.CreateVertex();
+			VertexPositions[NewVertexID] = Position;
+			VertexIDMap.Add(i, NewVertexID);
+		}
 
-			for (uint32 i = Range.begin(), I = Range.end(); i < I; ++i)
-			{
-				const auto& AOSVertex = AOSMesh.Vertexes[Indexes[i]];
-				NormalArray[i] = AOSVertex.Normal;
-			}
-		});
+		checkSlow(VertexPositions.GetNumElements() == DstNumPositions);
 	}
 
-	// Makeup other data.
-	// DJH - Note I should be able to 'transfer' the FaceMatieralIdx, but this will require generating
-	// Auxiliary data because I only have this data on verts right now.
-	{
-		const uint32 DstNumTris = DstNumIndexes / 3;
-
-		OutRawMesh.FaceMaterialIndices.Empty(DstNumTris);
-		OutRawMesh.FaceMaterialIndices.AddZeroed(DstNumTris);
+	const uint32* AOSIndexes = AOSMesh.Indexes;
 
-		OutRawMesh.FaceSmoothingMasks.Empty(DstNumTris);
-		OutRawMesh.FaceSmoothingMasks.AddZeroed(DstNumTris);
+	// Connectivity: 
+	auto CreateTriangle = [&OutRawMesh, PolygonGroupID, &VertexInstanceNormals, &VertexInstanceTangents, &VertexInstanceBinormalSigns, &VertexInstanceColors, &VertexInstanceUVs, &EdgeHardnesses, &EdgeCreaseSharpnesses](const FVertexID TriangleIndex[3], const FVector Normals[3])
+	{
+		FVertexInstanceID VertexInstanceIDs[3];
+		for (int32 Corner = 0; Corner < 3; ++Corner)
+		{
+			VertexInstanceIDs[Corner] = OutRawMesh.CreateVertexInstance(TriangleIndex[Corner]);
+			VertexInstanceTangents[VertexInstanceIDs[Corner]] = FVector(1, 0, 0);
+			VertexInstanceNormals[VertexInstanceIDs[Corner]] = Normals[Corner];
+			VertexInstanceBinormalSigns[VertexInstanceIDs[Corner]] = GetBasisDeterminantSign(VertexInstanceTangents[VertexInstanceIDs[Corner]].GetSafeNormal(),
+																							 (VertexInstanceNormals[VertexInstanceIDs[Corner]] ^ VertexInstanceTangents[VertexInstanceIDs[Corner]]).GetSafeNormal(),
+																							 VertexInstanceNormals[VertexInstanceIDs[Corner]].GetSafeNormal());
+			VertexInstanceColors[VertexInstanceIDs[Corner]] = FVector4(1.0f);
+			VertexInstanceUVs.Set(VertexInstanceIDs[Corner], 0, FVector2D(0.0f, 0.0f));
+		}
 
-		OutRawMesh.WedgeColors.Empty();
-		OutRawMesh.WedgeColors.AddZeroed(3 * DstNumTris);
-		// what about red?
-		for (uint32 i = 0; i < 3 * DstNumTris; ++i)
+		//Create a polygon from this triangle
+		TArray<FMeshDescription::FContourPoint> Contours;
+		for (int32 Corner = 0; Corner < 3; ++Corner)
 		{
-			OutRawMesh.WedgeColors[i] = FColor(255, 0, 0, 0);
+			int32 ContourPointIndex = Contours.AddDefaulted();
+			FMeshDescription::FContourPoint& ContourPoint = Contours[ContourPointIndex];
+			//Find the matching edge ID
+			uint32 CornerIndices[2];
+			CornerIndices[0] = (Corner + 0) % 3;
+			CornerIndices[1] = (Corner + 1) % 3;
+
+			FVertexID EdgeVertexIDs[2];
+			EdgeVertexIDs[0] = TriangleIndex[CornerIndices[0]];
+			EdgeVertexIDs[1] = TriangleIndex[CornerIndices[1]];
+
+			FEdgeID MatchEdgeId = OutRawMesh.GetVertexPairEdge(EdgeVertexIDs[0], EdgeVertexIDs[1]);
+			if (MatchEdgeId == FEdgeID::Invalid)
+			{
+				MatchEdgeId = OutRawMesh.CreateEdge(EdgeVertexIDs[0], EdgeVertexIDs[1]);
+				EdgeHardnesses[MatchEdgeId] = false;
+				EdgeCreaseSharpnesses[MatchEdgeId] = 0.0f;
+			}
+			ContourPoint.EdgeID = MatchEdgeId;
+			ContourPoint.VertexInstanceID = VertexInstanceIDs[CornerIndices[0]];
 		}
+		// Insert a polygon into the mesh
+		const FPolygonID NewPolygonID = OutRawMesh.CreatePolygon(PolygonGroupID, Contours);
+		//Triangulate the polygon
+		FMeshPolygon& Polygon = OutRawMesh.GetPolygon(NewPolygonID);
+		OutRawMesh.ComputePolygonTriangulation(NewPolygonID, Polygon.Triangles);
+	};
 
-		for (int32 TexCoordIndex = 0; TexCoordIndex < 8; ++TexCoordIndex)
+	{
+		uint32 IndexStop = DstNumIndexes/3;
+		for (uint32 t = 0, T = IndexStop; t < T; ++t)
 		{
-			OutRawMesh.WedgeTexCoords[TexCoordIndex].Empty();
+			FVertexID VertexIndexes[3];
+			FVector Normals[3];
+			for (int32 Corner = 0; Corner < 3; ++Corner)
+			{
+				VertexIndexes[Corner] = VertexIDMap[AOSIndexes[(t*3) + Corner]];
+				const auto& AOSVertex = AOSMesh.Vertexes[AOSIndexes[(t*3) + Corner]];
+				Normals[Corner] = AOSVertex.Normal;
+			}
+			CreateTriangle(VertexIndexes, Normals);
 		}
-		// Need to have a a WedgeTextCoord channel
-		auto& WedgeTexCoord = OutRawMesh.WedgeTexCoords[0];
-		WedgeTexCoord.AddZeroed(DstNumIndexes);
 	}
 }
 
 
-void ProxyLOD::VertexDataMeshToRawMesh(const FVertexDataMesh& SrcVertexDataMesh, FRawMesh& DstRawMesh)
+void ProxyLOD::VertexDataMeshToRawMesh(const FVertexDataMesh& SrcVertexDataMesh, FMeshDescription& OutRawMesh)
 {
-	
+	TVertexAttributesRef<FVector> VertexPositions = OutRawMesh.VertexAttributes().GetAttributesRef<FVector>(MeshAttribute::Vertex::Position);
+	TEdgeAttributesRef<bool> EdgeHardnesses = OutRawMesh.EdgeAttributes().GetAttributesRef<bool>(MeshAttribute::Edge::IsHard);
+	TEdgeAttributesRef<float> EdgeCreaseSharpnesses = OutRawMesh.EdgeAttributes().GetAttributesRef<float>(MeshAttribute::Edge::CreaseSharpness);
+	TPolygonGroupAttributesRef<FName> PolygonGroupImportedMaterialSlotNames = OutRawMesh.PolygonGroupAttributes().GetAttributesRef<FName>(MeshAttribute::PolygonGroup::ImportedMaterialSlotName);
+	TVertexInstanceAttributesRef<FVector> VertexInstanceNormals = OutRawMesh.VertexInstanceAttributes().GetAttributesRef<FVector>(MeshAttribute::VertexInstance::Normal);
+	TVertexInstanceAttributesRef<FVector> VertexInstanceTangents = OutRawMesh.VertexInstanceAttributes().GetAttributesRef<FVector>(MeshAttribute::VertexInstance::Tangent);
+	TVertexInstanceAttributesRef<float> VertexInstanceBinormalSigns = OutRawMesh.VertexInstanceAttributes().GetAttributesRef<float>(MeshAttribute::VertexInstance::BinormalSign);
+	TVertexInstanceAttributesRef<FVector4> VertexInstanceColors = OutRawMesh.VertexInstanceAttributes().GetAttributesRef<FVector4>(MeshAttribute::VertexInstance::Color);
+	TVertexInstanceAttributesRef<FVector2D> VertexInstanceUVs = OutRawMesh.VertexInstanceAttributes().GetAttributesRef<FVector2D>(MeshAttribute::VertexInstance::TextureCoordinate);
 
 	const uint32 DstNumPositions = SrcVertexDataMesh.Points.Num();
 	const uint32 DstNumIndexes = SrcVertexDataMesh.Indices.Num();
-
-	checkSlow(DstNumIndexes % 3 == 0);
-	// Copy the vertices over
+	const uint32 SrcNumTriangles = DstNumIndexes / 3;
+	
+	if (VertexInstanceUVs.GetNumIndices() < 2)
 	{
-		// Allocate the space for the verts in the DstRawMesh
-
-		const TArray<FVector>& SrcPositions = SrcVertexDataMesh.Points;
-		TArray<FVector>& DstPositions = DstRawMesh.VertexPositions;
-
-		ResizeArray(DstPositions, DstNumPositions);
-
-
-		ProxyLOD::Parallel_For(ProxyLOD::FUIntRange(0, DstNumPositions),
-			[&SrcPositions, &DstPositions](const ProxyLOD::FUIntRange& Range)
-		{
-			for (uint32 i = Range.begin(), I = Range.end(); i < I; ++i)
-			{
-				DstPositions[i] = SrcPositions[i];
-			}
-		});
-
-		checkSlow(DstPositions.Num() == DstNumPositions);
+		//We set the lightmap channel so 2
+		VertexInstanceUVs.SetNumIndices(2);
 	}
 
-	// Connectivity: 
+	FPolygonGroupID PolygonGroupID = FPolygonGroupID::Invalid;
+	if (OutRawMesh.PolygonGroups().Num() == 0)
 	{
-		const TArray<uint32>& SrcIndices = SrcVertexDataMesh.Indices;
-		TArray<uint32>& DstIndices = DstRawMesh.WedgeIndices;
+		PolygonGroupID = OutRawMesh.CreatePolygonGroup();
+		PolygonGroupImportedMaterialSlotNames[PolygonGroupID] = FName(*FString::Printf(TEXT("ProxyLOD_Material_%d"), FMath::Rand()));
+	}
+	else
+	{
+		PolygonGroupID = OutRawMesh.PolygonGroups().GetFirstValidID();
+	}
 
-		ResizeArray(DstIndices, DstNumIndexes);
+	checkSlow(DstNumIndexes % 3 == 0);
+	// Copy the vertices over
+	TMap<int32, FVertexID> VertexIDMap;
+	VertexIDMap.Reserve(DstNumPositions);
+	{
+		const TArray<FVector>& SrcPositions = SrcVertexDataMesh.Points;
 
-		ProxyLOD::Parallel_For(ProxyLOD::FUIntRange(0, DstNumIndexes),
-			[&SrcIndices, &DstIndices](const ProxyLOD::FUIntRange& Range)
+		for (uint32 i = 0, I = DstNumPositions; i < I; ++i)
 		{
-			for (uint32 i = Range.begin(), I = Range.end(); i < I; ++i)
-			{
-				DstIndices[i] = SrcIndices[i];
-			}
-		});
+			const FVector& Position = SrcPositions[i];
+			const FVertexID NewVertexID = OutRawMesh.CreateVertex();
+			VertexPositions[NewVertexID] = Position;
+			VertexIDMap.Add(i, NewVertexID);
+		}
 
-		checkSlow(DstIndices.Num() == DstNumIndexes);
+		checkSlow(VertexPositions.GetNumElements() == DstNumPositions);
 	}
+	
+	const bool bSrcHasTangentSpace = SrcVertexDataMesh.Tangent.Num() != 0 && SrcVertexDataMesh.BiTangent.Num() != 0 && SrcVertexDataMesh.Normal.Num() != 0;
 
-	// Copy the tangent space
+	// Connectivity:
+	auto CreateTriangle = [&OutRawMesh, &SrcVertexDataMesh, bSrcHasTangentSpace, &VertexIDMap, PolygonGroupID, &VertexInstanceNormals, &VertexInstanceTangents, &VertexInstanceBinormalSigns, &VertexInstanceColors, &VertexInstanceUVs, &EdgeHardnesses, &EdgeCreaseSharpnesses](const uint32 TriangleIndices[3])
 	{
-		const TArray<FVector>& SrcTangentArray = SrcVertexDataMesh.Tangent;
-		const TArray<FVector>& SrcBiTangentArray = SrcVertexDataMesh.BiTangent;
-		const TArray<FVector>& SrcNormalArray = SrcVertexDataMesh.Normal;
-
-		TArray<FVector>& DstTangentArray = DstRawMesh.WedgeTangentX;
-		TArray<FVector>& DstBiTangentArray = DstRawMesh.WedgeTangentY;
-		TArray<FVector>& DstNormalArray = DstRawMesh.WedgeTangentZ;
-
-		ResizeArray(DstTangentArray, DstNumIndexes);
-		ResizeArray(DstBiTangentArray, DstNumIndexes);
-		ResizeArray(DstNormalArray, DstNumIndexes);
-
-		const bool bSrcHasTangentSpace = SrcVertexDataMesh.Tangent.Num() != 0 && SrcVertexDataMesh.BiTangent.Num() != 0 && SrcVertexDataMesh.Normal.Num() != 0;
-		if (bSrcHasTangentSpace)
+		int32 TriangleIndex = TriangleIndices[0] / 3;
+		
+		FVertexID VertexIndexes[3];
+		FVertexInstanceID VertexInstanceIDs[3];
+		for (int32 Corner = 0; Corner < 3; ++Corner)
 		{
+			int32 SrcIndex = SrcVertexDataMesh.Indices[TriangleIndices[Corner]];
 
-			const auto& SrcIndexes = SrcVertexDataMesh.Indices;
+			VertexIndexes[Corner] = VertexIDMap[SrcIndex];
+			VertexInstanceIDs[Corner] = OutRawMesh.CreateVertexInstance(VertexIndexes[Corner]);
 
-			// Copy tangent space.  
-			// NB: The raw mesh potentially stores a different tangent space
-			// for each index.  For the vertex data mesh, the tangent space is per-vertex.
-			ProxyLOD::Parallel_For(ProxyLOD::FUIntRange(0, DstNumIndexes),
-				[&SrcTangentArray, &SrcIndexes, &DstTangentArray](const ProxyLOD::FUIntRange& Range)
+			//Tangents
+			if (bSrcHasTangentSpace)
 			{
-				for (uint32 i = Range.begin(), I = Range.end(); i < I; ++i)
-				{
-					DstTangentArray[i] = SrcTangentArray[SrcIndexes[i]];
-				}
-			});
-
-			ProxyLOD::Parallel_For(ProxyLOD::FUIntRange(0, DstNumIndexes),
-				[&SrcBiTangentArray, &SrcIndexes, &DstBiTangentArray](const ProxyLOD::FUIntRange& Range)
+				FVector Tangent = SrcVertexDataMesh.Tangent[SrcIndex];
+				FVector BiTangent = SrcVertexDataMesh.BiTangent[SrcIndex];
+				FVector Normal = SrcVertexDataMesh.Normal[SrcIndex];
+				VertexInstanceTangents[VertexInstanceIDs[Corner]] = Tangent;
+				VertexInstanceBinormalSigns[VertexInstanceIDs[Corner]] = GetBasisDeterminantSign(Tangent, BiTangent, Normal);
+				VertexInstanceNormals[VertexInstanceIDs[Corner]] = Normal;
+			}
+			else
 			{
-				for (uint32 i = Range.begin(), I = Range.end(); i < I; ++i)
-				{
-					DstBiTangentArray[i] = SrcBiTangentArray[SrcIndexes[i]];
-				}
-			});
+				VertexInstanceTangents[VertexInstanceIDs[Corner]] = FVector(1, 0, 0);
+				VertexInstanceNormals[VertexInstanceIDs[Corner]] = FVector(0, 0, 1);
+				VertexInstanceBinormalSigns[VertexInstanceIDs[Corner]] = GetBasisDeterminantSign(VertexInstanceTangents[VertexInstanceIDs[Corner]].GetSafeNormal(),
+																		 (VertexInstanceNormals[VertexInstanceIDs[Corner]] ^ VertexInstanceTangents[VertexInstanceIDs[Corner]]).GetSafeNormal(),
+																		 VertexInstanceNormals[VertexInstanceIDs[Corner]].GetSafeNormal());
+			}
 
-			ProxyLOD::Parallel_For(ProxyLOD::FUIntRange(0, DstNumIndexes),
-				[&SrcNormalArray, &SrcIndexes, &DstNormalArray](const ProxyLOD::FUIntRange& Range)
+			//Color
+			if (SrcVertexDataMesh.FaceColors.Num() == 0)
 			{
-				for (uint32 i = Range.begin(), I = Range.end(); i < I; ++i)
-				{
-					DstNormalArray[i] = SrcNormalArray[SrcIndexes[i]];
-				}
-			});
-		}
-		else // generate a default tangent space
-		{
-			for (uint32 i = 0; i < DstNumIndexes; ++i)
-			{
-				DstTangentArray[i] = FVector(1, 0, 0);
-				DstBiTangentArray[i] = FVector(0, 1, 0);
-				DstNormalArray[i] = FVector(0, 0, 1);
+				VertexInstanceColors[VertexInstanceIDs[Corner]] = FVector4(1.0f);
 			}
-		}
-	}
-
-	
-	// Is the handedness the same for all verts on a given face?
-	{
-		const uint32 DstNumTris = DstNumIndexes / 3;
-
-
-		ResizeArray(DstRawMesh.FaceMaterialIndices, DstNumTris);
-		ProxyLOD::Parallel_For(ProxyLOD::FUIntRange(0, DstNumTris), [&DstRawMesh](const ProxyLOD::FUIntRange& Range)
-		{
-			auto& FaceMaterialIndices = DstRawMesh.FaceMaterialIndices;
-			for (uint32 f = Range.begin(), F = Range.end(); f < F; ++f)
+			else
 			{
-				FaceMaterialIndices[f] = 0;
+				//There is 
+				VertexInstanceColors[VertexInstanceIDs[Corner]] = FVector4(FLinearColor(SrcVertexDataMesh.FaceColors[TriangleIndex]));
 			}
-		});
 
-		ResizeArray(DstRawMesh.FaceSmoothingMasks, DstNumTris);
-		if (SrcVertexDataMesh.FacePartition.Num() != 0)
-		{
-
-			ProxyLOD::Parallel_For(ProxyLOD::FUIntRange(0, DstNumTris), [&DstRawMesh, &SrcVertexDataMesh](const ProxyLOD::FUIntRange& Range)
+			//UVs, copy two time the same value, one for UV_Channel 0 and another for lightmap channel 1
+			for (int32 channel = 0; channel < 2; ++channel)
 			{
-				const auto& FacePartition = SrcVertexDataMesh.FacePartition;
-				auto& FaceSmoothingMask = DstRawMesh.FaceSmoothingMasks;
-				for (uint32 f = Range.begin(), F = Range.end(); f < F; ++f)
+				if (SrcVertexDataMesh.UVs.Num() == 0)
 				{
-					FaceSmoothingMask[f] = (1 << (FacePartition[f] % 32));
+					VertexInstanceUVs.Set(VertexInstanceIDs[Corner], channel, FVector2D(0.0f, 0.0f));
+				}
+				else
+				{
+					VertexInstanceUVs.Set(VertexInstanceIDs[Corner], channel, SrcVertexDataMesh.UVs[SrcIndex]);
 				}
-			});
-		}
-		else // default 
-		{
-			for (uint32 i = 0; i < DstNumTris; ++i)
-			{
-				DstRawMesh.FaceSmoothingMasks[i] = 1;
 			}
 		}
 
-		ResizeArray(DstRawMesh.WedgeColors, 3 * DstNumTris);
-		if (SrcVertexDataMesh.FaceColors.Num() == 0)
+		//Create a polygon from this triangle
+		TArray<FMeshDescription::FContourPoint> Contours;
+		for (int32 Corner = 0; Corner < 3; ++Corner)
 		{
-
-			// what about red?
-			for (uint32 i = 0; i < 3 * DstNumTris; ++i)
+			int32 ContourPointIndex = Contours.AddDefaulted();
+			FMeshDescription::FContourPoint& ContourPoint = Contours[ContourPointIndex];
+			//Find the matching edge ID
+			uint32 CornerIndices[2];
+			CornerIndices[0] = (Corner + 0) % 3;
+			CornerIndices[1] = (Corner + 1) % 3;
+
+			FVertexID EdgeVertexIDs[2];
+			EdgeVertexIDs[0] = VertexIndexes[CornerIndices[0]];
+			EdgeVertexIDs[1] = VertexIndexes[CornerIndices[1]];
+
+			FEdgeID MatchEdgeId = OutRawMesh.GetVertexPairEdge(EdgeVertexIDs[0], EdgeVertexIDs[1]);
+			if (MatchEdgeId == FEdgeID::Invalid)
 			{
-				DstRawMesh.WedgeColors[i] = FColor(255, 0, 0, 0);
+				MatchEdgeId = OutRawMesh.CreateEdge(EdgeVertexIDs[0], EdgeVertexIDs[1]);
+				EdgeHardnesses[MatchEdgeId] = false;
+				EdgeCreaseSharpnesses[MatchEdgeId] = 0.0f;
 			}
+			ContourPoint.EdgeID = MatchEdgeId;
+			ContourPoint.VertexInstanceID = VertexInstanceIDs[CornerIndices[0]];
 		}
-		else
+		// Insert a polygon into the mesh
+		const FPolygonID NewPolygonID = OutRawMesh.CreatePolygon(PolygonGroupID, Contours);
+		//Triangulate the polygon
+		FMeshPolygon& Polygon = OutRawMesh.GetPolygon(NewPolygonID);
+		OutRawMesh.ComputePolygonTriangulation(NewPolygonID, Polygon.Triangles);
+	};
+
+	{
+		uint32 RangeStop = DstNumIndexes / 3;
+		for (uint32 i = 0, I = RangeStop; i < I; ++i)
 		{
-			// Transfer color to the verts.
-			for (uint32 i = 0; i < DstNumTris; ++i)
+			uint32 SrcIndexes[3];
+			for (int32 Corner = 0; Corner < 3; ++Corner)
 			{
-				const uint32 WedgeId = i * 3;
-#if 1
-				const FColor& Color = SrcVertexDataMesh.FaceColors[i];
-
-				DstRawMesh.WedgeColors[WedgeId + 0] = Color;
-				DstRawMesh.WedgeColors[WedgeId + 1] = Color;
-				DstRawMesh.WedgeColors[WedgeId + 2] = Color;
-#else
-
-				uint32 Idx[3];
-				for (uint32 i = 0; i < 3; ++i) Idx[i] = WedgeId + i;
-
-				for (uint32 i = 0; i < 3; ++i) Idx[i] = SrcVertexDataMesh.Indices[Idx[i]];
-
-				for (uint32 i = 0; i < 3; ++i) DstRawMesh.WedgeColors[WedgeId + i] = (SrcVertexDataMesh.TangentHanded[Idx[i]] > 0) ? FColor(255, 0, 0) : FColor(0, 0, 255);
-#endif
+				SrcIndexes[Corner] = (i * 3) + Corner;
 			}
+			CreateTriangle(SrcIndexes);
 		}
-#if 0
-		
-		// Testing with the handedness
 
-		for (uint32 i = 0; i < DstNumIndexes; ++i)
-		{
-			uint32 v = DstRawMesh.WedgeIndices[i];
-			DstRawMesh.WedgeColors[i] = (SrcVertexDataMesh.TangentHanded[v] > 0) ? FColor(255, 0, 0) : FColor(0, 255, 0);
-		}
-#endif
+		checkSlow(OutRawMesh.VertexInstances().Num() == DstNumIndexes);
+	}
 
+	//Put everybody in the same smoothgroup by default
+	TArray<uint32> FaceSmoothingMasks;
+	FaceSmoothingMasks.AddZeroed(SrcNumTriangles);
 
+	if (SrcVertexDataMesh.FacePartition.Num() != 0)
+	{
+		for (uint32 FaceIndex = 0; FaceIndex < SrcNumTriangles; ++FaceIndex)
 		{
-			// Empty all the texture coords
-			for (int32 TexCoordIndex = 0; TexCoordIndex < 8; ++TexCoordIndex)
-			{
-				DstRawMesh.WedgeTexCoords[TexCoordIndex].Empty();
-			}
-
-			// Copy the UVs over to the first and second texture coord channel.  The second one is used by lightmass, 
-			// if we don't provide it, unreal will make it from the first for us
-			for (int32 channel = 0; channel < 2; ++channel)
-			{
-				// Need to have a a WedgeTextCoord channel
-				auto& WedgeTexCoord = DstRawMesh.WedgeTexCoords[channel];
-
-				if (SrcVertexDataMesh.UVs.Num() == 0)
-				{
-					WedgeTexCoord.AddZeroed(DstNumIndexes);
-				}
-				else
-				{
-					WedgeTexCoord.AddUninitialized(DstNumIndexes);
-					const uint32* Indices = SrcVertexDataMesh.Indices.GetData();
-					const auto& SrcUVs = SrcVertexDataMesh.UVs;
-					ProxyLOD::Parallel_For(ProxyLOD::FUIntRange(0, DstNumIndexes), [&Indices, &SrcUVs, &WedgeTexCoord](const ProxyLOD::FUIntRange& Range)
-					{
-						for (uint32 i = Range.begin(), I = Range.end(); i < I; ++i)
-						{
-							const uint32 VertId = Indices[i];
-							const FVector2D& UV = SrcUVs[VertId];
-
-							WedgeTexCoord[i] = UV;
-						}
-					});
-				}
-			}
+			FaceSmoothingMasks[FaceIndex] = (1 << (SrcVertexDataMesh.FacePartition[FaceIndex] % 32));
 		}
-
 	}
+
+	FMeshDescriptionOperations::ConvertSmoothGroupToHardEdges(FaceSmoothingMasks, OutRawMesh);
 }
 
 
@@ -486,126 +439,108 @@ void ProxyLOD::VertexDataMeshToRawMesh(const FVertexDataMesh& SrcVertexDataMesh,
 // raw mesh is nominally a per-index data structure and the vertex data mesh is a per-vertex structure.
 // In addition, this only transfers the first texture coordinate and ignores material ids and vertex colors.
 
-void ProxyLOD::RawMeshToVertexDataMesh(const FRawMesh& SrcRawMesh, FVertexDataMesh& DstVertexDataMesh)
+void ProxyLOD::RawMeshToVertexDataMesh(const FMeshDescription& SrcRawMesh, FVertexDataMesh& DstVertexDataMesh)
 {
+	TVertexAttributesConstRef<FVector> VertexPositions = SrcRawMesh.VertexAttributes().GetAttributesRef<FVector>(MeshAttribute::Vertex::Position);
+	TVertexInstanceAttributesConstRef<FVector> VertexInstanceNormals = SrcRawMesh.VertexInstanceAttributes().GetAttributesRef<FVector>(MeshAttribute::VertexInstance::Normal);
+	TVertexInstanceAttributesConstRef<FVector> VertexInstanceTangents = SrcRawMesh.VertexInstanceAttributes().GetAttributesRef<FVector>(MeshAttribute::VertexInstance::Tangent);
+	TVertexInstanceAttributesConstRef<float> VertexInstanceBinormalSigns = SrcRawMesh.VertexInstanceAttributes().GetAttributesRef<float>(MeshAttribute::VertexInstance::BinormalSign);
+	TVertexInstanceAttributesConstRef<FVector2D> VertexInstanceUVs = SrcRawMesh.VertexInstanceAttributes().GetAttributesRef<FVector2D>(MeshAttribute::VertexInstance::TextureCoordinate);
 
 
-	const uint32 DstNumPositions = SrcRawMesh.VertexPositions.Num();
-	const uint32 DstNumIndexes = SrcRawMesh.WedgeIndices.Num();
+	const uint32 DstNumPositions = SrcRawMesh.Vertices().Num();
 
-	checkSlow(DstNumIndexes % 3 == 0);
-	// Copy the vertices over
+	uint32 DstNumIndexes = 0;
+	for (const FPolygonID& PolygonID : SrcRawMesh.Polygons().GetElementIDs())
 	{
-		// Allocate the space for the verts in the VertexDataMesh
-
-		ResizeArray(DstVertexDataMesh.Points, DstNumPositions);
-
-		ProxyLOD::Parallel_For(ProxyLOD::FUIntRange(0, DstNumPositions),
-			[&DstVertexDataMesh, &SrcRawMesh](const ProxyLOD::FUIntRange& Range)
-		{
-			const FVector* SrcPos = SrcRawMesh.VertexPositions.GetData();
-			FVector* DstPos = DstVertexDataMesh.Points.GetData();
-			for (uint32 i = Range.begin(), I = Range.end(); i < I; ++i)
-			{
-				DstPos[i] = SrcPos[i];
-
-			}
-		});
-
-		checkSlow(DstVertexDataMesh.Points.Num() == DstNumPositions);
+		const FMeshPolygon& Polygon = SrcRawMesh.GetPolygon(PolygonID);
+		DstNumIndexes += Polygon.Triangles.Num() * 3;
 	}
 
-	// Connectivity: 
+	// Copy the vertices over
+	TMap<FVertexID, uint32> VertexIDToDstVertexIndex;
+	VertexIDToDstVertexIndex.Reserve(SrcRawMesh.Vertices().Num());
 	{
-		// convert each quad into two triangles.
-		const TArray<uint32>& SrcIndices = SrcRawMesh.WedgeIndices;
-		TArray<uint32>& DstIndices = DstVertexDataMesh.Indices;
-		ResizeArray(DstIndices, DstNumIndexes);
-
-		ProxyLOD::Parallel_For(ProxyLOD::FUIntRange(0, DstNumIndexes),
-			[&DstIndices, &SrcIndices](const ProxyLOD::FUIntRange& Range)
-		{;
-		for (uint32 i = Range.begin(), I = Range.end(); i < I; ++i)
+		// Allocate the space for the verts in the VertexDataMesh
+		ResizeArray(DstVertexDataMesh.Points, DstNumPositions);
+		uint32 VertexCount = 0;
+		for (const FVertexID& VertexID : SrcRawMesh.Vertices().GetElementIDs())
 		{
-			DstIndices[i] = SrcIndices[i];
+			DstVertexDataMesh.Points[VertexCount] = VertexPositions[VertexID];
+			VertexIDToDstVertexIndex.Add(VertexID, VertexCount);
+			VertexCount++;
 		}
-		});
-
-		checkSlow(DstIndices.Num() == DstNumIndexes);
 	}
 
-	// Copy the tangent space:
-	// NB: The tangent space is stored per-index in the raw mesh, but only per-vertex in the vertex data mesh.
-	// We assume that the raw mesh per-index data is really duplicated per-vertex data!
-
-	{
-		// Note, these aren't threaded to avoid possible data-races.
-		const bool bThreaded = false;
+	// Connectivity:
+	uint32 VertexInstanceCount = 0;
+	TArray<uint32>& DstIndices = DstVertexDataMesh.Indices;
+	ResizeArray(DstIndices, DstNumIndexes);
 
+	TArray<FVector>& DstTangentArray = DstVertexDataMesh.Tangent;
+	ResizeArray(DstTangentArray, DstNumPositions);
 
-		const TArray<FVector>& SrcTangentArray = SrcRawMesh.WedgeTangentX;
-		const TArray<FVector>& SrcBiTangentArray = SrcRawMesh.WedgeTangentY;
-		const TArray<FVector>& SrcNormalArray = SrcRawMesh.WedgeTangentZ;
+	TArray<FVector>& DstBiTangentArray = DstVertexDataMesh.BiTangent;
+	ResizeArray(DstBiTangentArray, DstNumPositions);
 
-		TArray<FVector>& DstTangentArray = DstVertexDataMesh.Tangent;
-		TArray<FVector>& DstBiTangentArray = DstVertexDataMesh.BiTangent;
-		TArray<FVector>& DstNormalArray = DstVertexDataMesh.Normal;
+	TArray<FVector>& DstNormalArray = DstVertexDataMesh.Normal;
+	ResizeArray(DstNormalArray, DstNumPositions);
 
-		ResizeArray(DstTangentArray, DstNumPositions);
-		ResizeArray(DstBiTangentArray, DstNumPositions);
-		ResizeArray(DstNormalArray, DstNumPositions);
+	TArray<FVector2D>& DstUVs = DstVertexDataMesh.UVs;
+	ResizeArray(DstUVs, DstNumPositions);
 
-		const auto& DstIndexes = DstVertexDataMesh.Indices;
-		// Copy tangent space
-		ProxyLOD::Parallel_For(ProxyLOD::FUIntRange(0, DstNumIndexes),
-			[&DstIndexes, &SrcTangentArray, &DstTangentArray](const ProxyLOD::FUIntRange& Range)
-		{
-			for (uint32 i = Range.begin(), I = Range.end(); i < I; ++i)
-			{
-				DstTangentArray[DstIndexes[i]] = SrcTangentArray[i];
-			}
-		}, bThreaded);
-
-		ProxyLOD::Parallel_For(ProxyLOD::FUIntRange(0, DstNumIndexes),
-			[&DstIndexes, &DstBiTangentArray, &SrcBiTangentArray](const ProxyLOD::FUIntRange& Range)
+	//Iterate all triangle and add the indices
+	for (const FPolygonID& PolygonID : SrcRawMesh.Polygons().GetElementIDs())
+	{
+		const FMeshPolygon& Polygon = SrcRawMesh.GetPolygon(PolygonID);
+		for (const FMeshTriangle& Triangle : Polygon.Triangles)
 		{
-			for (uint32 i = Range.begin(), I = Range.end(); i < I; ++i)
+			for (int32 Corner = 0; Corner < 3; ++Corner)
 			{
-				DstBiTangentArray[DstIndexes[i]] = SrcBiTangentArray[i];
-			}
-		}, bThreaded);
+				const FVertexInstanceID& VertexInstanceID = Triangle.GetVertexInstanceID(Corner);
+				DstIndices[VertexInstanceCount] = VertexIDToDstVertexIndex[SrcRawMesh.GetVertexInstanceVertex(VertexInstanceID)];
+
+				// Copy the tangent space:
+				// NB: The tangent space is stored per-index in the raw mesh, but only per-vertex in the vertex data mesh.
+				// We assume that the raw mesh per-index data is really duplicated per-vertex data!
+				DstTangentArray[DstIndices[VertexInstanceCount]] = VertexInstanceTangents[VertexInstanceID];
+				DstBiTangentArray[DstIndices[VertexInstanceCount]] = FVector::CrossProduct(VertexInstanceNormals[VertexInstanceID], VertexInstanceTangents[VertexInstanceID]).GetSafeNormal() * VertexInstanceBinormalSigns[VertexInstanceID];
+				DstNormalArray[DstIndices[VertexInstanceCount]] = VertexInstanceNormals[VertexInstanceID];
+
+				// Copy the UVs:
+				// NB: The UVs is stored per-index in the raw mesh, but only per-vertex in the vertex data mesh.
+				// We assume that the raw mesh per-index data is really duplicated per-vertex data!
+				if (VertexInstanceUVs.GetNumIndices() == 0)
+				{
+					DstUVs[DstIndices[VertexInstanceCount]] = FVector2D(0.0f, 0.0f);
+				}
+				else
+				{
+					DstUVs[DstIndices[VertexInstanceCount]] = VertexInstanceUVs.Get(VertexInstanceID, 0);
+				}
 
-		ProxyLOD::Parallel_For(ProxyLOD::FUIntRange(0, DstNumIndexes),
-			[&DstIndexes, &DstNormalArray, &SrcNormalArray](const ProxyLOD::FUIntRange& Range)
-		{
-			for (uint32 i = Range.begin(), I = Range.end(); i < I; ++i)
-			{
-				DstNormalArray[DstIndexes[i]] = SrcNormalArray[i];
+				VertexInstanceCount++;
 			}
-		}, bThreaded);
+		}
 	}
 
+	int32 NumTriangles = VertexInstanceCount / 3;
+	TArray<uint32> FaceSmoothingMasks;
+	FaceSmoothingMasks.AddZeroed(NumTriangles);
 
-	// Copy UVs
-	{
-		// Note, these aren't threaded to avoid possible data-races.
-		const bool bThreaded = false;
-
-		// Need to have a a WedgeTextCoord channel
-		auto& SrcUVs = SrcRawMesh.WedgeTexCoords[0];
-		TArray<FVector2D>& DstUVs = DstVertexDataMesh.UVs;
+	FMeshDescriptionOperations::ConvertHardEdgesToSmoothGroup(SrcRawMesh, FaceSmoothingMasks);
 
-		ResizeArray(DstUVs, DstNumPositions);
+	TArray<uint32>& DstFacePartition = DstVertexDataMesh.FacePartition;
+	ResizeArray(DstFacePartition, NumTriangles);
 
-		const auto& DstIndexes = DstVertexDataMesh.Indices;
-		ProxyLOD::Parallel_For(ProxyLOD::FUIntRange(0, DstNumIndexes),
-			[&DstIndexes, &DstUVs, &SrcUVs](const ProxyLOD::FUIntRange& Range)
+	for (int32 FaceIndex = 0; FaceIndex < NumTriangles; ++FaceIndex)
+	{
+		DstFacePartition[FaceIndex] = 0;
+		for (int32 BitIndex = 0; BitIndex < 32; ++BitIndex)
 		{
-			for (uint32 i = Range.begin(), I = Range.end(); i < I; ++i)
-			{
-				DstUVs[DstIndexes[i]] = SrcUVs[i];
-			}
-		}, bThreaded);
+			uint32 BitMask = FMath::Pow(2, BitIndex);
+			DstFacePartition[FaceIndex] += ((FaceSmoothingMasks[FaceIndex] & BitMask) > 0);
+		}
 	}
 }
 
@@ -730,17 +665,17 @@ void ProxyLOD::ConvertMesh(const FAOSMesh& InMesh, FVertexDataMesh& OutMesh)
 }
 
 
-void ProxyLOD::ConvertMesh(const FAOSMesh& InMesh, FRawMesh& OutMesh)
+void ProxyLOD::ConvertMesh(const FAOSMesh& InMesh, FMeshDescription& OutMesh)
 {
 	AOSMeshToRawMesh(InMesh, OutMesh);
 }
 
-void ProxyLOD::ConvertMesh(const FVertexDataMesh& InMesh, FRawMesh& OutMesh)
+void ProxyLOD::ConvertMesh(const FVertexDataMesh& InMesh, FMeshDescription& OutMesh)
 {
 	VertexDataMeshToRawMesh(InMesh, OutMesh);
 }
 
-void ProxyLOD::ConvertMesh(const FRawMesh& InMesh, FVertexDataMesh& OutMesh)
+void ProxyLOD::ConvertMesh(const FMeshDescription& InMesh, FVertexDataMesh& OutMesh)
 {
 	RawMeshToVertexDataMesh(InMesh, OutMesh);
 }
diff --git a/Engine/Plugins/Experimental/ProxyLODPlugin/Source/ProxyLOD/Private/ProxyLODMeshConvertUtils.h b/Engine/Plugins/Experimental/ProxyLODPlugin/Source/ProxyLOD/Private/ProxyLODMeshConvertUtils.h
--- a/Engine/Plugins/Experimental/ProxyLODPlugin/Source/ProxyLOD/Private/ProxyLODMeshConvertUtils.h
+++ b/Engine/Plugins/Experimental/ProxyLODPlugin/Source/ProxyLOD/Private/ProxyLODMeshConvertUtils.h
@@ -28,7 +28,7 @@ namespace ProxyLOD
 	* @param InMesh   Mesh with a mixture of quads and triangles.
 	* @param OutMesh  Triangle Mesh.
 	*/
-	void MixedPolyMeshToRawMesh( const FMixedPolyMesh& InMesh, FRawMesh& OutMesh );
+	void MixedPolyMeshToRawMesh( const FMixedPolyMesh& InMesh, FMeshDescription& OutMesh );
 	
 	/**
 	* Convert a mixed triangle and quad mesh to a triangle mesh type by splitting quads.
@@ -49,33 +49,33 @@ namespace ProxyLOD
 	*
 	* The conversions will maintain geometry and connectivity, but different mesh types support
 	* different attribute types and possibly different frequenceies even when the same attribute types
-	* exist on both mesh types ( e.g. FRawMesh has per-wedge data while FVertexDataMesh has per-vertex attributes).
+	* exist on both mesh types ( e.g. FMeshDescription has per-wedge data while FVertexDataMesh has per-vertex attributes).
 	*
 	* @param InMesh   Source Mesh to convert.
 	* @param OutMesh  The result of the mesh conversion.  Any data already stored in OutMesh will be lost.
 	*
 	*/
-	void ConvertMesh(const FAOSMesh& InMesh, FRawMesh& OutMesh);
+	void ConvertMesh(const FAOSMesh& InMesh, FMeshDescription& OutMesh);
 	void ConvertMesh(const FAOSMesh& InMesh, FVertexDataMesh& OutMesh);
-	void ConvertMesh(const FVertexDataMesh& InMesh, FRawMesh& OutMesh);
-	void ConvertMesh(const FRawMesh& InMesh, FVertexDataMesh& OutMesh);
+	void ConvertMesh(const FVertexDataMesh& InMesh, FMeshDescription& OutMesh);
+	void ConvertMesh(const FMeshDescription& InMesh, FVertexDataMesh& OutMesh);
 
 	
 	/**
-	* Convert the simplifier-friendly array-of-structs mesh to a struct of arrays FRawMesh.
+	* Convert the simplifier-friendly array-of-structs mesh to a struct of arrays FMeshDescription.
 	*
-	* NB: This copies the AOSVertex normal to the wedge FRawMesh::Tangentz.
-	*     But additional FRawMesh attributes, including tangent/bitangent are given default values.
+	* NB: This copies the AOSVertex normal to the wedge FMeshDescription::Tangentz.
+	*     But additional FMeshDescription attributes, including tangent/bitangent are given default values.
 	*
 	* @param InMesh   Source Mesh to convert.
 	* @param OutMesh  The result of the mesh conversion.  Any data already stored in OutMesh will be lost.
 	*
 	*/
-	void AOSMeshToRawMesh( const FAOSMesh& InMesh, FRawMesh& OutMesh );
+	void AOSMeshToRawMesh( const FAOSMesh& InMesh, FMeshDescription& OutMesh );
 
 
 	/**
-	* Convert the uv-generation-friendly vertex data mesh to a struct of arrays FRawMesh.
+	* Convert the uv-generation-friendly vertex data mesh to a struct of arrays FMeshDescription.
 	* In addition to connectivity and vertex locations, this also transfers tangent space
 	* and UVs.
 	*
@@ -83,10 +83,10 @@ namespace ProxyLOD
 	* @param OutMesh  The result of the mesh conversion.  Any data already stored in OutMesh will be lost.
 	*
 	*/
-	void VertexDataMeshToRawMesh( const FVertexDataMesh& InMesh, FRawMesh& OutMesh );
+	void VertexDataMeshToRawMesh( const FVertexDataMesh& InMesh, FMeshDescription& OutMesh );
 
 	/**
-	* Converts a FRawMesh to a  uv-generation-friendly vertex data mesh.  This is potentially has some loss since the 
+	* Converts a FMeshDescription to a  uv-generation-friendly vertex data mesh.  This is potentially has some loss since the 
 	* raw mesh is nominally a per-index data structure and the vertex data mesh is a per-vertex structure.
 	* In addition, this only transfers the first texture coordinate and ignores material ids and vertex colors.
 	*
@@ -94,7 +94,7 @@ namespace ProxyLOD
 	* @param OutMesh  The result of the mesh conversion.  Any data already stored in OutMesh will be lost.
 	*
 	*/
-	void RawMeshToVertexDataMesh( const FRawMesh& InMesh, FVertexDataMesh& OutMesh );
+	void RawMeshToVertexDataMesh( const FMeshDescription& InMesh, FVertexDataMesh& OutMesh );
 
 }
 
@@ -223,7 +223,7 @@ void ProxyLOD::MixedPolyMeshToAOSMesh(const FMixedPolyMesh& MixedPolyMesh, TAOSM
 
 
 #if 0
-static void ComputeAABB(const FRawMeshArrayAdapter& MeshArray, ProxyLOD::FBBox& BBox)
+static void ComputeAABB(const FMeshDescriptionArrayAdapter& MeshArray, ProxyLOD::FBBox& BBox)
 {
 	uint32 NumTris = MeshArray.polygonCount();
 	BBox = ProxyLOD::Parallel_Reduce(ProxyLOD::FIntRange(0, NumTris), ProxyLOD::FBBox(), [&MeshArray](const ProxyLOD::FIntRange& Range,  ProxyLOD::FBBox BBox)->ProxyLOD::FBBox
diff --git a/Engine/Plugins/Experimental/ProxyLODPlugin/Source/ProxyLOD/Private/ProxyLODMeshParameterization.cpp b/Engine/Plugins/Experimental/ProxyLODPlugin/Source/ProxyLOD/Private/ProxyLODMeshParameterization.cpp
--- a/Engine/Plugins/Experimental/ProxyLODPlugin/Source/ProxyLOD/Private/ProxyLODMeshParameterization.cpp
+++ b/Engine/Plugins/Experimental/ProxyLODPlugin/Source/ProxyLOD/Private/ProxyLODMeshParameterization.cpp
@@ -262,30 +262,43 @@ void ProxyLOD::GenerateAdjacency(const FVertexDataMesh& Mesh, std::vector<uint32
 	HRESULT hr = DirectX::GenerateAdjacencyAndPointReps(Mesh.Indices.GetData(), NumTris, (const DirectX::XMFLOAT3*)PosArray, NumPos, Eps, NULL /* optional point rep pointer*/, AdjacencyArray.data());
 
 }
-
-void ProxyLOD::GenerateAdjacency(const FRawMesh& RawMesh, std::vector<uint32>& AdjacencyArray)
+void ProxyLOD::GenerateAdjacency(const FMeshDescription& RawMesh, std::vector<uint32>& AdjacencyArray)
 {
-	const uint32 NumTris = RawMesh.WedgeIndices.Num() / 3;
-	const uint32 NumVerts = RawMesh.VertexPositions.Num();
-	const uint32 AdjacencySize = RawMesh.WedgeIndices.Num(); // = 3 for each face
+	uint32 NumTris = 0;
+	for (const FPolygonID& PolygonID : RawMesh.Polygons().GetElementIDs())
+	{
+		NumTris += RawMesh.GetPolygonTriangles(PolygonID).Num();
+	}
+	const uint32 NumVerts = RawMesh.Vertices().Num();
+	const uint32 AdjacencySize = RawMesh.VertexInstances().Num(); // 3 for each face
 
 															 // Allocate adjacency
 	AdjacencyArray.resize(AdjacencySize);
+
+	TVertexAttributesConstRef<FVector> VertexPositionsAttribute = RawMesh.VertexAttributes().GetAttributesRef<FVector>(MeshAttribute::Vertex::Position);
+	TArray<FVector> VertexPositions;
+	VertexPositions.AddZeroed(NumVerts);
+	for (const FVertexID VertexID : RawMesh.Vertices().GetElementIDs())
+	{
+		VertexPositions[VertexID.GetValue()] = VertexPositionsAttribute[VertexID];
+	}
+
+	TArray<uint32> Indices;
+	Indices.AddZeroed(AdjacencySize);
+
+	for (const FVertexInstanceID VertexInstanceID : RawMesh.VertexInstances().GetElementIDs())
+	{
+		Indices[VertexInstanceID.GetValue()] = RawMesh.GetVertexInstanceVertex(VertexInstanceID).GetValue();
+	}
 	// position comparison epsilon
 	const float Eps = 0.f;
-	HRESULT hr = DirectX::GenerateAdjacencyAndPointReps(RawMesh.WedgeIndices.GetData(), NumTris, (DirectX::XMFLOAT3*)RawMesh.VertexPositions.GetData(), NumVerts, Eps, NULL /* optional point rep pointer*/, AdjacencyArray.data());
+	HRESULT hr = DirectX::GenerateAdjacencyAndPointReps(Indices.GetData(), NumTris, (DirectX::XMFLOAT3*)VertexPositions.GetData(), NumVerts, Eps, NULL /* optional point rep pointer*/, AdjacencyArray.data());
 }
 
-
 bool ProxyLOD::GenerateAdjacenyAndCleanMesh(FVertexDataMesh& InOutMesh, std::vector<uint32>& Adjacency)
 {
-	
-
-
 	std::vector<uint32_t> dupVerts;
 
-	
-
 	uint32 CleanCount = 0;
 	while (CleanCount == 0 || (dupVerts.size() != 0 && CleanCount < 5))
 	{
diff --git a/Engine/Plugins/Experimental/ProxyLODPlugin/Source/ProxyLOD/Private/ProxyLODMeshParameterization.h b/Engine/Plugins/Experimental/ProxyLODPlugin/Source/ProxyLOD/Private/ProxyLODMeshParameterization.h
--- a/Engine/Plugins/Experimental/ProxyLODPlugin/Source/ProxyLOD/Private/ProxyLODMeshParameterization.h
+++ b/Engine/Plugins/Experimental/ProxyLODPlugin/Source/ProxyLOD/Private/ProxyLODMeshParameterization.h
@@ -54,5 +54,5 @@ namespace ProxyLOD
 	*/
 	void GenerateAdjacency(const FVertexDataMesh& InMesh, std::vector<uint32>& OutAdjacencyArray);
 	void GenerateAdjacency(const FAOSMesh& InMesh, std::vector<uint32>& OutAdjacencyArray);
-	void GenerateAdjacency(const FRawMesh& InMesh, std::vector<uint32>& OutAdjacencyArray);
+	void GenerateAdjacency(const FMeshDescription& InMesh, std::vector<uint32>& OutAdjacencyArray);
 }
\ No newline at end of file
diff --git a/Engine/Plugins/Experimental/ProxyLODPlugin/Source/ProxyLOD/Private/ProxyLODMeshSDFConversions.cpp b/Engine/Plugins/Experimental/ProxyLODPlugin/Source/ProxyLOD/Private/ProxyLODMeshSDFConversions.cpp
--- a/Engine/Plugins/Experimental/ProxyLODPlugin/Source/ProxyLOD/Private/ProxyLODMeshSDFConversions.cpp
+++ b/Engine/Plugins/Experimental/ProxyLODPlugin/Source/ProxyLOD/Private/ProxyLODMeshSDFConversions.cpp
@@ -12,7 +12,7 @@
 //#include <exception>
 
 
-bool ProxyLOD::MeshArrayToSDFVolume(const FRawMeshArrayAdapter& MeshAdapter, openvdb::FloatGrid::Ptr& SDFGrid, openvdb::Int32Grid* PolyIndexGrid)
+bool ProxyLOD::MeshArrayToSDFVolume(const FMeshDescriptionArrayAdapter& MeshAdapter, openvdb::FloatGrid::Ptr& SDFGrid, openvdb::Int32Grid* PolyIndexGrid)
 {
 
 	bool success = true;
@@ -56,7 +56,7 @@ bool ProxyLOD::MeshArrayToSDFVolume(const FRawMeshArrayAdapter& MeshAdapter, ope
 }
 
 
-bool ProxyLOD::MeshArrayToSDFVolume(const FRawMeshAdapter& MeshAdapter, openvdb::FloatGrid::Ptr& SDFGrid, openvdb::Int32Grid* PolyIndexGrid)
+bool ProxyLOD::MeshArrayToSDFVolume(const FMeshDescriptionAdapter& MeshAdapter, openvdb::FloatGrid::Ptr& SDFGrid, openvdb::Int32Grid* PolyIndexGrid)
 {
 
 	bool success = true;
@@ -116,8 +116,9 @@ static openvdb::FloatGrid::Ptr OffsetSDF(const openvdb::FloatGrid::Ptr InSDFVolu
 	FMixedPolyMesh  MixedPolyMesh;
 	openvdb::tools::volumeToMesh(*InSDFVolume, MixedPolyMesh.Points, MixedPolyMesh.Triangles, MixedPolyMesh.Quads, IsoValue, 0.001);
 
-	// convert the mesh to FRawMesh
-	FRawMesh RawMesh;
+	// convert the mesh to FMeshDescription
+	FMeshDescription RawMesh;
+	UStaticMesh::RegisterMeshAttributes(RawMesh);
 	ProxyLOD::MixedPolyMeshToRawMesh(MixedPolyMesh, RawMesh);
 
 	// Create a new empty grid with the same transform and metadata
@@ -125,7 +126,7 @@ static openvdb::FloatGrid::Ptr OffsetSDF(const openvdb::FloatGrid::Ptr InSDFVolu
 	OutSDFVolume->setTransform(openvdb::math::Transform::createLinearTransform(ResultVolexSize));
 
 	// Wrap so we can re-voxelize with bandwidth 2
-	FRawMeshAdapter  MeshAdapter(RawMesh, OutSDFVolume->transform());
+	FMeshDescriptionAdapter  MeshAdapter(RawMesh, OutSDFVolume->transform());
 
 	// Re-voxelize
 	ProxyLOD::MeshArrayToSDFVolume(MeshAdapter, OutSDFVolume);
diff --git a/Engine/Plugins/Experimental/ProxyLODPlugin/Source/ProxyLOD/Private/ProxyLODMeshSDFConversions.h b/Engine/Plugins/Experimental/ProxyLODPlugin/Source/ProxyLOD/Private/ProxyLODMeshSDFConversions.h
--- a/Engine/Plugins/Experimental/ProxyLODPlugin/Source/ProxyLOD/Private/ProxyLODMeshSDFConversions.h
+++ b/Engine/Plugins/Experimental/ProxyLODPlugin/Source/ProxyLOD/Private/ProxyLODMeshSDFConversions.h
@@ -28,11 +28,11 @@ namespace ProxyLOD
 	*
 	* @retun 'true' is success, 'false' will generally be an out of memory error.
 	*/
-	bool MeshArrayToSDFVolume( const FRawMeshArrayAdapter& InMeshAdapter,
+	bool MeshArrayToSDFVolume( const FMeshDescriptionArrayAdapter& InMeshAdapter,
 	                           openvdb::FloatGrid::Ptr& OutSDFGrid,
 		                       openvdb::Int32Grid* OutPolyIndexGrid = nullptr);
 
-	bool MeshArrayToSDFVolume( const FRawMeshAdapter& InMeshAdapter,
+	bool MeshArrayToSDFVolume( const FMeshDescriptionAdapter& InMeshAdapter,
 		                       openvdb::FloatGrid::Ptr& OutSDFGrid,
 		                       openvdb::Int32Grid* OutPolyIndexGrid = nullptr);
 
@@ -111,7 +111,7 @@ template <typename DstMeshType>
 void ProxyLOD::SDFVolumeToMesh(const openvdb::FloatGrid::ConstPtr SDFVolume, const double IsoValue, const double Adaptivity, DstMeshType& OutMesh)
 {
 
-	// we should be generating a new FRawMesh.
+	// we should be generating a new FMeshDescription.
 
 	OutMesh.Empty();
 
diff --git a/Engine/Plugins/Experimental/ProxyLODPlugin/Source/ProxyLOD/Private/ProxyLODMeshTypes.cpp b/Engine/Plugins/Experimental/ProxyLODPlugin/Source/ProxyLOD/Private/ProxyLODMeshTypes.cpp
--- a/Engine/Plugins/Experimental/ProxyLODPlugin/Source/ProxyLOD/Private/ProxyLODMeshTypes.cpp
+++ b/Engine/Plugins/Experimental/ProxyLODPlugin/Source/ProxyLOD/Private/ProxyLODMeshTypes.cpp
@@ -1,44 +1,67 @@
 // Copyright 1998-2018 Epic Games, Inc. All Rights Reserved.
 
 #include "ProxyLODMeshTypes.h"
+#include "MeshDescriptionOperations.h"
 
-// --- FRawMeshAdapter ----
+// --- FMeshDescriptionAdapter ----
 
-FRawMeshAdapter::FRawMeshAdapter(const FRawMesh& InRawMesh, const openvdb::math::Transform& InTransform) :
+FMeshDescriptionAdapter::FMeshDescriptionAdapter(const FMeshDescription& InRawMesh, const openvdb::math::Transform& InTransform) :
 	RawMesh(&InRawMesh), Transform(&InTransform)
-{}
+{
+	InitializeCacheData();
+}
 
-FRawMeshAdapter::FRawMeshAdapter(const FRawMeshAdapter& other)
+FMeshDescriptionAdapter::FMeshDescriptionAdapter(const FMeshDescriptionAdapter& other)
 	: RawMesh(other.RawMesh), Transform(other.Transform)
-{}
+{
+	InitializeCacheData();
+}
 
-size_t FRawMeshAdapter::polygonCount() const
+void FMeshDescriptionAdapter::InitializeCacheData()
 {
-	// RawMesh hold triangles
-	return size_t(RawMesh->WedgeIndices.Num() / 3);
+	VertexPositions = RawMesh->VertexAttributes().GetAttributesRef<FVector>(MeshAttribute::Vertex::Position);
+	TriangleCount = 0;
+	for (const FPolygonID& PolygonID : RawMesh->Polygons().GetElementIDs())
+	{
+		TriangleCount += RawMesh->GetPolygonTriangles(PolygonID).Num();
+	}
 }
 
-size_t FRawMeshAdapter::pointCount() const
+size_t FMeshDescriptionAdapter::polygonCount() const
 {
-	return size_t(RawMesh->VertexPositions.Num());
+	return size_t(TriangleCount);
 }
 
-void FRawMeshAdapter::getIndexSpacePoint(size_t FaceNumber, size_t CornerNumber, openvdb::Vec3d& pos) const
+size_t FMeshDescriptionAdapter::pointCount() const
 {
-	// Get the vertex position in local space.
-	const int32 WedgeIndex = int32(FaceNumber * 3 + CornerNumber);
+	return size_t(RawMesh->Vertices().Num());
+}
 
+void FMeshDescriptionAdapter::getIndexSpacePoint(size_t FaceNumber, size_t CornerNumber, openvdb::Vec3d& pos) const
+{
+	// Get the vertex position in local space.
+	const FVertexInstanceID VertexInstanceID = FVertexInstanceID(FaceNumber * 3 + CornerNumber);
 	// float3 position 
-	FVector Position = RawMesh->GetWedgePosition(WedgeIndex);
+	FVector Position = VertexPositions[RawMesh->GetVertexInstanceVertex(VertexInstanceID)];
 	pos = Transform->worldToIndex(openvdb::Vec3d(Position.X, Position.Y, Position.Z));
-
 };
 
 
-// --- FRawMeshArrayAdapter ----
-
+// --- FMeshDescriptionArrayAdapter ----
+FMeshDescription* GetMeshDescriptionFromMeshMergeData(const FMeshMergeData* MergeData)
+{
+	FMeshDescription *MeshDescription = new FMeshDescription();
+	UStaticMesh::RegisterMeshAttributes(*MeshDescription);
+	TMap<int32, FName> MaterialMap;
+	for (int32 MaterialIndex = 0; MaterialIndex < MergeData->SourceStaticMesh->StaticMaterials.Num(); ++MaterialIndex)
+	{
+		MaterialMap.Add(MaterialIndex, MergeData->SourceStaticMesh->StaticMaterials[MaterialIndex].ImportedMaterialSlotName);
+	}
+	FMeshDescriptionOperations::ConvertFromRawMesh(*(MergeData->RawMesh), *MeshDescription, MaterialMap);
+	return MeshDescription;
+}
 
-FRawMeshArrayAdapter::FRawMeshArrayAdapter(const TArray<const FMeshMergeData*>& InMergeDataPtrArray)
+FMeshDescriptionArrayAdapter::FMeshDescriptionArrayAdapter(const TArray<const FMeshMergeData*>& InMergeDataPtrArray)
 {
 	// Make a default transform.
 	Transform = openvdb::math::Transform::createLinearTransform(1.);
@@ -50,13 +73,17 @@ FRawMeshArrayAdapter::FRawMeshArrayAdapter(const TArray<const FMeshMergeData*>&
 	for (int32 MeshIdx = 0, MeshCount = InMergeDataPtrArray.Num(); MeshIdx < MeshCount; ++MeshIdx)
 	{
 		const FMeshMergeData* MergeData = InMergeDataPtrArray[MeshIdx];
-		const FRawMesh* RawMesh = MergeData->RawMesh;
-		PointCount += size_t(RawMesh->VertexPositions.Num());
-		PolyCount += size_t(RawMesh->WedgeIndices.Num() / 3);
+		FMeshDescription *RawMesh = GetMeshDescriptionFromMeshMergeData(MergeData);
+
+		PointCount += size_t(RawMesh->Vertices().Num());
+		for (const FPolygonID& PolygonID : RawMesh->Polygons().GetElementIDs())
+		{
+			PolyCount += RawMesh->GetPolygonTriangles(PolygonID).Num();
+		}
 
 		PolyOffsetArray.push_back(PolyCount);
 		RawMeshArray.push_back(RawMesh);
-
+		RawMeshArrayData.push_back(FMeshDescriptionAttributesGetter(RawMesh));
 		MergeDataArray.push_back(MergeData);
 	}
 
@@ -64,7 +91,7 @@ FRawMeshArrayAdapter::FRawMeshArrayAdapter(const TArray<const FMeshMergeData*>&
 	ComputeAABB(this->BBox);
 }
 
-FRawMeshArrayAdapter::FRawMeshArrayAdapter(const TArray<FMeshMergeData>& InMergeDataArray)
+FMeshDescriptionArrayAdapter::FMeshDescriptionArrayAdapter(const TArray<FMeshMergeData>& InMergeDataArray)
 {
 	// Make a default transform.
 	Transform = openvdb::math::Transform::createLinearTransform(1.);
@@ -76,13 +103,17 @@ FRawMeshArrayAdapter::FRawMeshArrayAdapter(const TArray<FMeshMergeData>& InMerge
 	for (int32 MeshIdx = 0, MeshCount = InMergeDataArray.Num(); MeshIdx < MeshCount; ++MeshIdx)
 	{
 		const FMeshMergeData* MergeData = &InMergeDataArray[MeshIdx];
-		const FRawMesh* RawMesh = MergeData->RawMesh;
-		PointCount += size_t(RawMesh->VertexPositions.Num());
-		PolyCount += size_t(RawMesh->WedgeIndices.Num() / 3);
+		
+		FMeshDescription *RawMesh = GetMeshDescriptionFromMeshMergeData(MergeData);
+		PointCount += size_t(RawMesh->Vertices().Num());
+		for (const FPolygonID& PolygonID : RawMesh->Polygons().GetElementIDs())
+		{
+			PolyCount += RawMesh->GetPolygonTriangles(PolygonID).Num();
+		}
 
 		PolyOffsetArray.push_back(PolyCount);
 		RawMeshArray.push_back(RawMesh);
-
+		RawMeshArrayData.push_back(FMeshDescriptionAttributesGetter(RawMesh));
 		MergeDataArray.push_back(MergeData);
 	}
 
@@ -90,7 +121,7 @@ FRawMeshArrayAdapter::FRawMeshArrayAdapter(const TArray<FMeshMergeData>& InMerge
 	ComputeAABB(this->BBox);
 }
 
-FRawMeshArrayAdapter::FRawMeshArrayAdapter(const TArray<FMeshMergeData>& InMergeDataArray, const openvdb::math::Transform::Ptr InTransform)
+FMeshDescriptionArrayAdapter::FMeshDescriptionArrayAdapter(const TArray<FMeshMergeData>& InMergeDataArray, const openvdb::math::Transform::Ptr InTransform)
 	:Transform(InTransform)
 {
 	PointCount = 0;
@@ -100,13 +131,17 @@ FRawMeshArrayAdapter::FRawMeshArrayAdapter(const TArray<FMeshMergeData>& InMerge
 	for (int32 MeshIdx = 0, MeshCount = InMergeDataArray.Num(); MeshIdx < MeshCount; ++MeshIdx)
 	{
 		const FMeshMergeData* MergeData = &InMergeDataArray[MeshIdx];
-		const FRawMesh* RawMesh = MergeData->RawMesh;
-		PointCount += size_t(RawMesh->VertexPositions.Num());
-		PolyCount += size_t(RawMesh->WedgeIndices.Num() / 3);
+		
+		FMeshDescription *RawMesh = GetMeshDescriptionFromMeshMergeData(MergeData);
+		PointCount += size_t(RawMesh->Vertices().Num());
+		for (const FPolygonID& PolygonID : RawMesh->Polygons().GetElementIDs())
+		{
+			PolyCount += RawMesh->GetPolygonTriangles(PolygonID).Num();
+		}
 
 		PolyOffsetArray.push_back(PolyCount);
 		RawMeshArray.push_back(RawMesh);
-
+		RawMeshArrayData.push_back(FMeshDescriptionAttributesGetter(RawMesh));
 		MergeDataArray.push_back(MergeData);
 	}
 
@@ -114,31 +149,48 @@ FRawMeshArrayAdapter::FRawMeshArrayAdapter(const TArray<FMeshMergeData>& InMerge
 	ComputeAABB(this->BBox);
 }
 
-FRawMeshArrayAdapter::FRawMeshArrayAdapter(const FRawMeshArrayAdapter& other)
+FMeshDescriptionArrayAdapter::FMeshDescriptionArrayAdapter(const FMeshDescriptionArrayAdapter& other)
 	:Transform(other.Transform), PointCount(other.PointCount), PolyCount(other.PolyCount), BBox(other.BBox)
 {
 	RawMeshArray = other.RawMeshArray;
 	PolyOffsetArray = other.PolyOffsetArray;
 	MergeDataArray = other.MergeDataArray;
+
+	for (const FMeshDescription* RawMesh : RawMeshArray)
+	{
+		RawMeshArrayData.push_back(FMeshDescriptionAttributesGetter(RawMesh));
+	}
 }
 
-void FRawMeshArrayAdapter::getWorldSpacePoint(size_t FaceNumber, size_t CornerNumber, openvdb::Vec3d& pos) const
+FMeshDescriptionArrayAdapter::~FMeshDescriptionArrayAdapter()
 {
-	int32 MeshIdx, LocalFaceNumber;
+	for (FMeshDescription* MeshDescription : RawMeshArray)
+	{
+		delete MeshDescription;
+	}
+	RawMeshArray.clear();
 
-	const FRawMesh& RawMesh = GetRawMesh(FaceNumber, MeshIdx, LocalFaceNumber);
+	RawMeshArrayData.clear();
+	MergeDataArray.clear();
+	PolyOffsetArray.clear();
+}
 
+void FMeshDescriptionArrayAdapter::getWorldSpacePoint(size_t FaceNumber, size_t CornerNumber, openvdb::Vec3d& pos) const
+{
+	int32 MeshIdx, LocalFaceNumber;
 
-	// Get the vertex position in local space.
-	const int32 WedgeIndex = 3 * LocalFaceNumber + int32(CornerNumber);
+	const FMeshDescriptionAttributesGetter* AttributesGetter = nullptr;
+	const FMeshDescription& RawMesh = GetRawMesh(FaceNumber, MeshIdx, LocalFaceNumber, &AttributesGetter);
+	check(AttributesGetter);
 
+	// Get the vertex position in local space.
+	const FVertexInstanceID VertexInstanceID = FVertexInstanceID(3 * LocalFaceNumber + int32(CornerNumber));
 	// float3 position 
-	FVector Position = RawMesh.GetWedgePosition(WedgeIndex);
+	FVector Position = AttributesGetter->VertexPositions[RawMesh.GetVertexInstanceVertex(VertexInstanceID)];
 	pos = openvdb::Vec3d(Position.X, Position.Y, Position.Z);
-
 };
 
-void FRawMeshArrayAdapter::getIndexSpacePoint(size_t FaceNumber, size_t CornerNumber, openvdb::Vec3d& pos) const
+void FMeshDescriptionArrayAdapter::getIndexSpacePoint(size_t FaceNumber, size_t CornerNumber, openvdb::Vec3d& pos) const
 {
 	openvdb::Vec3d Position;
 	getWorldSpacePoint(FaceNumber, CornerNumber, Position);
@@ -146,75 +198,102 @@ void FRawMeshArrayAdapter::getIndexSpacePoint(size_t FaceNumber, size_t CornerNu
 
 };
 
-const FMeshMergeData& FRawMeshArrayAdapter::GetMeshMergeData(uint32 Idx) const
+const FMeshMergeData& FMeshDescriptionArrayAdapter::GetMeshMergeData(uint32 Idx) const
 {
 	checkSlow(Idx < MergeDataArray.size());
 	return *MergeDataArray[Idx];
 }
 
-FRawMeshArrayAdapter::FRawPoly FRawMeshArrayAdapter::GetRawPoly(const size_t FaceNumber, int32& OutMeshIdx, int32& OutLocalFaceNumber) const
+void FMeshDescriptionArrayAdapter::UpdateMaterialsID()
+{
+	for (int32 MeshIdx = 0; MeshIdx < MergeDataArray.size(); ++MeshIdx)
+	{
+		FMeshDescription* MeshDescription = RawMeshArray[MeshIdx];
+
+		int32 FaceNum = MergeDataArray[MeshIdx]->RawMesh->FaceMaterialIndices.Num();
+		check(FaceNum == RawMeshArrayData[MeshIdx].TriangleCount);
+		TMap<FPolygonGroupID, FPolygonGroupID> RemapGroup;
+		TArray<int32> UniqueMaterials;
+		int32 FaceIndex = 0;
+		for (const FPolygonID PolygonID : MeshDescription->Polygons().GetElementIDs())
+		{
+			int32 NewMaterialIndex = MergeDataArray[MeshIdx]->RawMesh->FaceMaterialIndices[FaceIndex];
+			if (!UniqueMaterials.Contains(NewMaterialIndex))
+			{
+				UniqueMaterials.Add(NewMaterialIndex);
+				FPolygonGroupID NewPolygonGroupID = FPolygonGroupID(NewMaterialIndex);
+				FPolygonGroupID OriginalPolygonGroupID = MeshDescription->GetPolygonPolygonGroup(PolygonID);
+				RemapGroup.Add(OriginalPolygonGroupID, NewPolygonGroupID);
+			}
+			//Update the face index
+			FaceIndex += MeshDescription->GetPolygonTriangles(PolygonID).Num();
+		}
+		//Remap the polygon group with the correct ID
+		FMeshDescriptionOperations::RemapPolygonGroups(*MeshDescription, RemapGroup);
+	}
+}
+
+FMeshDescriptionArrayAdapter::FRawPoly FMeshDescriptionArrayAdapter::GetRawPoly(const size_t FaceNumber, int32& OutMeshIdx, int32& OutLocalFaceNumber) const
 {
 	checkSlow(FaceNumber < PolyCount);
 
 	int32 MeshIdx, LocalFaceNumber;
 
-	const FRawMesh& RawMesh = GetRawMesh(FaceNumber, MeshIdx, LocalFaceNumber);
+	const FMeshDescriptionAttributesGetter* AttributesGetter = nullptr;
+	const FMeshDescription& RawMesh = GetRawMesh(FaceNumber, MeshIdx, LocalFaceNumber, &AttributesGetter);
+	check(AttributesGetter);
 	OutMeshIdx = MeshIdx;
 	OutLocalFaceNumber = LocalFaceNumber;
-	checkSlow(LocalFaceNumber < RawMesh.WedgeIndices.Num() / 3)
 
-		int32 WedgeIdx[3];
-	WedgeIdx[0] = 3 * LocalFaceNumber;
-	WedgeIdx[1] = WedgeIdx[0] + 1;
-	WedgeIdx[2] = WedgeIdx[0] + 2;
+	checkSlow(LocalFaceNumber < AttributesGetter->TriangleCount);
+
+	FVertexInstanceID WedgeIdx[3];
+	WedgeIdx[0] = FVertexInstanceID(3 * LocalFaceNumber);
+	WedgeIdx[1] = FVertexInstanceID((3 * LocalFaceNumber) + 1);
+	WedgeIdx[2] = FVertexInstanceID((3 * LocalFaceNumber) + 2);
 
 
 	FRawPoly RawPoly;
 	RawPoly.MeshIdx = MeshIdx;
-	RawPoly.FaceMaterialIndex = RawMesh.FaceMaterialIndices[LocalFaceNumber];
-	RawPoly.FaceSmoothingMask = RawMesh.FaceSmoothingMasks[LocalFaceNumber];
+	FPolygonID PolygonID(LocalFaceNumber);
+	RawPoly.FaceMaterialIndex = RawMesh.GetPolygonPolygonGroup(PolygonID).GetValue();
+	RawPoly.FaceSmoothingMask = AttributesGetter->FaceSmoothingMasks[LocalFaceNumber];
 
-	for (int32 i = 0; i < 3; ++i)  RawPoly.VertexPositions[i] = RawMesh.GetWedgePosition(WedgeIdx[i]);
-
-
-	if (RawMesh.WedgeTangentX.Num() > 0)
-		for (int32 i = 0; i < 3; ++i)
-			RawPoly.WedgeTangentX[i] = RawMesh.WedgeTangentX[WedgeIdx[i]];
-
-	if (RawMesh.WedgeTangentY.Num() > 0)
+	for (const FMeshTriangle& Triangle : RawMesh.GetPolygonTriangles(PolygonID))
+	{
+		FVertexInstanceID VertexInstanceIDs[3];
 		for (int32 i = 0; i < 3; ++i)
-			RawPoly.WedgeTangentY[i] = RawMesh.WedgeTangentY[WedgeIdx[i]];
+		{
+			VertexInstanceIDs[i] = Triangle.GetVertexInstanceID(i);
+			RawPoly.VertexPositions[i] = AttributesGetter->VertexPositions[RawMesh.GetVertexInstanceVertex(VertexInstanceIDs[i])];
+		}
 
-	if (RawMesh.WedgeTangentZ.Num() > 0)
-		for (int32 i = 0; i < 3; ++i)
-			RawPoly.WedgeTangentZ[i] = RawMesh.WedgeTangentZ[WedgeIdx[i]];
 
-	if (RawMesh.WedgeColors.Num() > 0)
 		for (int32 i = 0; i < 3; ++i)
-			RawPoly.WedgeColors[i] = RawMesh.WedgeColors[WedgeIdx[i]];
-	else
-		for (int32 i = 0; i < 3; ++i)
-			RawPoly.WedgeColors[i] = FColor::White;
-
-	// Copy Texture coords
-	for (int Idx = 0; Idx < MAX_MESH_TEXTURE_COORDS; ++Idx)
-	{
-
-		const TArray<FVector2D>& WedgeTexCoords = RawMesh.WedgeTexCoords[Idx];
-
-		if (WedgeTexCoords.Num() > 0)
-			for (int32 i = 0; i < 3; ++i)
+		{
+			RawPoly.WedgeTangentX[i] = AttributesGetter->VertexInstanceTangents[VertexInstanceIDs[i]];
+			RawPoly.WedgeTangentY[i] = FVector::CrossProduct(AttributesGetter->VertexInstanceNormals[VertexInstanceIDs[i]], AttributesGetter->VertexInstanceTangents[VertexInstanceIDs[i]]).GetSafeNormal() * AttributesGetter->VertexInstanceBinormalSigns[VertexInstanceIDs[i]];
+			RawPoly.WedgeTangentZ[i] = AttributesGetter->VertexInstanceNormals[VertexInstanceIDs[i]];
+			
+			RawPoly.WedgeColors[i] = FLinearColor(AttributesGetter->VertexInstanceColors[VertexInstanceIDs[i]]).ToFColor(true);
+			// Copy Texture coords
+			for (int Idx = 0; Idx < MAX_MESH_TEXTURE_COORDS_MD; ++Idx)
 			{
-				RawPoly.WedgeTexCoords[Idx][i] = WedgeTexCoords[WedgeIdx[i]];
+				if (AttributesGetter->VertexInstanceUVs.GetNumIndices() > Idx)
+				{
+					RawPoly.WedgeTexCoords[Idx][i] = AttributesGetter->VertexInstanceUVs.Get(VertexInstanceIDs[i], Idx);
+				}
+				else
+				{
+					RawPoly.WedgeTexCoords[Idx][i] = FVector2D(0.f, 0.f);
+				}
 			}
-		else
-			for (int32 i = 0; i < 3; ++i)
-				RawPoly.WedgeTexCoords[Idx][i] = FVector2D(0.f, 0.f);
+		}
 	}
 	return RawPoly;
 }
 
-FRawMeshArrayAdapter::FRawPoly FRawMeshArrayAdapter::GetRawPoly(const size_t FaceNumber) const
+FMeshDescriptionArrayAdapter::FRawPoly FMeshDescriptionArrayAdapter::GetRawPoly(const size_t FaceNumber) const
 {
 	int32 IgnoreMeshId, IgnoreLocalFaceNumber;
 	return GetRawPoly(FaceNumber, IgnoreMeshId, IgnoreLocalFaceNumber);
@@ -222,7 +301,7 @@ FRawMeshArrayAdapter::FRawPoly FRawMeshArrayAdapter::GetRawPoly(const size_t Fac
 
 // protected functions
 
-const FRawMesh& FRawMeshArrayAdapter::GetRawMesh(const size_t FaceNumber, int32& MeshIdx, int32& LocalFaceNumber) const
+const FMeshDescription& FMeshDescriptionArrayAdapter::GetRawMesh(const size_t FaceNumber, int32& MeshIdx, int32& LocalFaceNumber, const FMeshDescriptionAttributesGetter** OutAttributesGetter) const
 {
 	// Find the correct raw mesh
 	MeshIdx = 0;
@@ -234,12 +313,14 @@ const FRawMesh& FRawMeshArrayAdapter::GetRawMesh(const size_t FaceNumber, int32&
 	// Offset the face number to get the correct index into this mesh.
 	LocalFaceNumber = int32(FaceNumber) - PolyOffsetArray[MeshIdx];
 
-	const FRawMesh* RawMesh = RawMeshArray[MeshIdx];
+	const FMeshDescription* MeshDescription = RawMeshArray[MeshIdx];
 
-	return *RawMesh;
+	*OutAttributesGetter = &RawMeshArrayData[MeshIdx];
+	
+	return *MeshDescription;
 }
 
-void FRawMeshArrayAdapter::ComputeAABB(ProxyLOD::FBBox& InOutBBox)
+void FMeshDescriptionArrayAdapter::ComputeAABB(ProxyLOD::FBBox& InOutBBox)
 {
 	uint32 NumTris = this->polygonCount();
 	InOutBBox = ProxyLOD::Parallel_Reduce(ProxyLOD::FIntRange(0, NumTris), ProxyLOD::FBBox(),
@@ -273,7 +354,7 @@ void FRawMeshArrayAdapter::ComputeAABB(ProxyLOD::FBBox& InOutBBox)
 }
 
 // --- FClosestPolyField ----
-FClosestPolyField::FClosestPolyField(const FRawMeshArrayAdapter& MeshArray, const openvdb::Int32Grid::Ptr& SrcPolyIndexGrid) :
+FClosestPolyField::FClosestPolyField(const FMeshDescriptionArrayAdapter& MeshArray, const openvdb::Int32Grid::Ptr& SrcPolyIndexGrid) :
 	RawMeshArrayAdapter(&MeshArray),
 	ClosestPolyGrid(SrcPolyIndexGrid) 
 {}
@@ -283,14 +364,14 @@ FClosestPolyField::FClosestPolyField(const FClosestPolyField& other) :
 	ClosestPolyGrid(other.ClosestPolyGrid) 
 {}
 
-FClosestPolyField::FPolyConstAccessor::FPolyConstAccessor(const openvdb::Int32Grid* PolyIndexGrid, const FRawMeshArrayAdapter* MeshArrayAdapter) :
+FClosestPolyField::FPolyConstAccessor::FPolyConstAccessor(const openvdb::Int32Grid* PolyIndexGrid, const FMeshDescriptionArrayAdapter* MeshArrayAdapter) :
 	MeshArray(MeshArrayAdapter),
 	CAccessor(PolyIndexGrid->getConstAccessor()),
 	XForm(&(PolyIndexGrid->transform()))
 {
 }
 
-FRawMeshArrayAdapter::FRawPoly FClosestPolyField::FPolyConstAccessor::Get(const openvdb::Vec3d& WorldPos, bool& bSuccess) const
+FMeshDescriptionArrayAdapter::FRawPoly FClosestPolyField::FPolyConstAccessor::Get(const openvdb::Vec3d& WorldPos, bool& bSuccess) const
 {
 	checkSlow(MeshArray != NULL);
 	const openvdb::Coord ijk = XForm->worldToIndexCellCentered(WorldPos);
diff --git a/Engine/Plugins/Experimental/ProxyLODPlugin/Source/ProxyLOD/Private/ProxyLODMeshTypes.h b/Engine/Plugins/Experimental/ProxyLODPlugin/Source/ProxyLOD/Private/ProxyLODMeshTypes.h
--- a/Engine/Plugins/Experimental/ProxyLODPlugin/Source/ProxyLOD/Private/ProxyLODMeshTypes.h
+++ b/Engine/Plugins/Experimental/ProxyLODPlugin/Source/ProxyLOD/Private/ProxyLODMeshTypes.h
@@ -10,8 +10,11 @@
 
 #endif
 
-#include "RawMesh.h"
 #include "MeshMergeData.h"
+#include "MeshDescription.h"
+#include "MeshAttributes.h"
+#include "MeshAttributeArray.h"
+#include "MeshDescriptionOperations.h"
 
 #include "ProxyLODThreadedWrappers.h"
 #include "ProxyLODVertexTypes.h"
@@ -22,15 +25,15 @@
 
 /**
 * Various stages in the ProxyLOD pipeline require different internal mesh formats, but
-* both input and output data for the ProxyLOD process use the native UE format, FRawMesh.
+* both input and output data for the ProxyLOD process use the native UE format, FMeshDescription.
 *
 *  
 *
-* FRawMeshAdapter:
-*	An adapter class to allow a FRawMesh to be voxelized without actually converting to a new mesh format.
+* FMeshDescriptionAdapter:
+*	An adapter class to allow a FMeshDescription to be voxelized without actually converting to a new mesh format.
 *
-* FRawMeshArrayAdapter:
-*	An adapter class to allow an array of FRawMesh items to be viewed as a single mesh.  Primarly to allow
+* FMeshDescriptionArrayAdapter:
+*	An adapter class to allow an array of FMeshDescription items to be viewed as a single mesh.  Primarly to allow
 *   multiple items of geometry to be voxelized into a single grid.
 * 
 * FMixedPolyMesh:
@@ -189,14 +192,14 @@ typedef TAOSMesh<FPositionOnlyVertex>    FPosSimplifierMesh;
 *  Mesh adapter class that implements the needed API for the methods in openvdb that convert
 *  polygonal mesh to a signed distance field.
 */
-class FRawMeshAdapter
+class FMeshDescriptionAdapter
 {
 public:
 	/**
 	* Constructor 
 	*/
-	FRawMeshAdapter(const FRawMesh& InRawMesh, const openvdb::math::Transform& InTransform);
-	FRawMeshAdapter(const FRawMeshAdapter& other);
+	FMeshDescriptionAdapter(const FMeshDescription& InRawMesh, const openvdb::math::Transform& InTransform);
+	FMeshDescriptionAdapter(const FMeshDescriptionAdapter& other);
 
 	// Total number of polygons
 	size_t polygonCount() const;
@@ -204,7 +207,7 @@ public:
 	// Total number of points (vertex locations)
 	size_t pointCount() const;
 	
-	// Vertex count for polygon n: currently FRawMesh is just triangles.
+	// Vertex count for polygon n: currently FMeshDescription is just triangles.
 	size_t vertexCount(size_t n) const { return 3; }
 
 	// Return position pos in local grid index space for polygon n and vertex v
@@ -221,7 +224,13 @@ public:
 private:
 
 	// Pointer to the raw mesh we are wrapping
-	const FRawMesh* RawMesh;
+	const FMeshDescription* RawMesh;
+
+	//////////////////////////////////////////////////////////////////////////
+	//Cache data
+	void InitializeCacheData();
+	uint32 TriangleCount;
+	TVertexAttributesConstRef<FVector> VertexPositions;
 
 	// Transform used to convert the mesh space into the index space used by voxelization
 	const openvdb::math::Transform* Transform;
@@ -239,18 +248,21 @@ namespace ProxyLOD
 * NB: This is only an adapter class that holds pointers to external objects.  It does not manage the lifetimes of those objects,
 *
 */
-class FRawMeshArrayAdapter
+class FMeshDescriptionArrayAdapter
 {
 public:
 
-	FRawMeshArrayAdapter(const TArray<FMeshMergeData>& InMergeDataArray, const openvdb::math::Transform::Ptr InTransform);
+	FMeshDescriptionArrayAdapter(const TArray<FMeshMergeData>& InMergeDataArray, const openvdb::math::Transform::Ptr InTransform);
 
-	FRawMeshArrayAdapter(const TArray<const FMeshMergeData*>& InMergeDataPtrArray);
+	FMeshDescriptionArrayAdapter(const TArray<const FMeshMergeData*>& InMergeDataPtrArray);
 	
-	FRawMeshArrayAdapter(const TArray<FMeshMergeData>& InMergeDataArray);
+	FMeshDescriptionArrayAdapter(const TArray<FMeshMergeData>& InMergeDataArray);
 
 	// copy constructor 
-	FRawMeshArrayAdapter(const FRawMeshArrayAdapter& other);
+	FMeshDescriptionArrayAdapter(const FMeshDescriptionArrayAdapter& other);
+
+	//Destructor
+	~FMeshDescriptionArrayAdapter();
 	
 	// Return position for polygon (face number) n and vertex (corner number) v
 	void getWorldSpacePoint(size_t FaceNumber, size_t CornerNumber, openvdb::Vec3d& pos) const;
@@ -262,12 +274,15 @@ public:
 
 	// Access to the FMeshMergeData data elements in the array.
 	const FMeshMergeData& GetMeshMergeData(uint32 Idx) const;
+
+	//Set the FMeshMergeData material into the FMeshDescription
+	void UpdateMaterialsID();
 	
 	/**
 	* Triangle class that contains the information associated with a single triangle
-	* in an FRawMesh.
+	* in an FMeshDescription.
 	*  
-	*  NB: The FRawMesh is a struct of arrays, this basically converts the data for a single 
+	*  NB: The FMeshDescription is a struct of arrays, this basically converts the data for a single 
 	*      poly into a struct.
 	*/
 	class FRawPoly
@@ -284,11 +299,41 @@ public:
 		FVector   WedgeTangentY[3];
 		FVector   WedgeTangentZ[3];
 
-		FVector2D WedgeTexCoords[MAX_MESH_TEXTURE_COORDS][3];
+		FVector2D WedgeTexCoords[MAX_MESH_TEXTURE_COORDS_MD][3];
 
 		FColor WedgeColors[3];
 	};
 
+	class FMeshDescriptionAttributesGetter
+	{
+	public:
+		FMeshDescriptionAttributesGetter(const FMeshDescription* MeshDescription)
+		{
+			VertexPositions = MeshDescription->VertexAttributes().GetAttributesRef<FVector>(MeshAttribute::Vertex::Position);
+			VertexInstanceNormals = MeshDescription->VertexInstanceAttributes().GetAttributesRef<FVector>(MeshAttribute::VertexInstance::Normal);
+			VertexInstanceTangents = MeshDescription->VertexInstanceAttributes().GetAttributesRef<FVector>(MeshAttribute::VertexInstance::Tangent);
+			VertexInstanceBinormalSigns = MeshDescription->VertexInstanceAttributes().GetAttributesRef<float>(MeshAttribute::VertexInstance::BinormalSign);
+			VertexInstanceColors = MeshDescription->VertexInstanceAttributes().GetAttributesRef<FVector4>(MeshAttribute::VertexInstance::Color);
+			VertexInstanceUVs = MeshDescription->VertexInstanceAttributes().GetAttributesRef<FVector2D>(MeshAttribute::VertexInstance::TextureCoordinate);
+
+			TriangleCount = 0;
+			for (FPolygonID PolygonID : MeshDescription->Polygons().GetElementIDs())
+			{
+				TriangleCount += MeshDescription->GetPolygonTriangles(PolygonID).Num();
+			}
+			FaceSmoothingMasks.AddZeroed(TriangleCount);
+			FMeshDescriptionOperations::ConvertHardEdgesToSmoothGroup(*MeshDescription, FaceSmoothingMasks);
+		}
+		TVertexAttributesConstRef<FVector> VertexPositions;
+		TVertexInstanceAttributesConstRef<FVector> VertexInstanceNormals;
+		TVertexInstanceAttributesConstRef<FVector> VertexInstanceTangents;
+		TVertexInstanceAttributesConstRef<float> VertexInstanceBinormalSigns;
+		TVertexInstanceAttributesConstRef<FVector4> VertexInstanceColors;
+		TVertexInstanceAttributesConstRef<FVector2D> VertexInstanceUVs;
+		int32 TriangleCount;
+		TArray<uint32> FaceSmoothingMasks;
+	};
+
 	/**
 	* Returns a copy of the data associated with this poly in the form of a struct.
 	* 
@@ -298,7 +343,7 @@ public:
 	*
 	* @return  A copy of the raw mesh data associated with this poly.
 	*/
-	FRawMeshArrayAdapter::FRawPoly GetRawPoly(const size_t FaceNumber, int32& OutMeshIdx, int32& OutLocalFaceNumber) const;
+	FMeshDescriptionArrayAdapter::FRawPoly GetRawPoly(const size_t FaceNumber, int32& OutMeshIdx, int32& OutLocalFaceNumber) const;
 	
 	/**
 	* Returns a copy of the data associated with this poly in the form of a struct.
@@ -307,7 +352,7 @@ public:
 	*
 	* @return  A copy of the raw mesh data associated with this poly.
 	*/
-	FRawMeshArrayAdapter::FRawPoly GetRawPoly(const size_t FaceNumber) const;
+	FMeshDescriptionArrayAdapter::FRawPoly GetRawPoly(const size_t FaceNumber) const;
 
 	/**
 	* Total number of polygons managed by this class.
@@ -326,7 +371,7 @@ public:
 	}
 
 	/**
-	* Vertex count for polygon n: currently FRawMesh is just triangles.
+	* Vertex count for polygon n: currently FMeshDescription is just triangles.
 	*/
 	size_t vertexCount(size_t n) const 
 	{ 
@@ -355,7 +400,7 @@ public:
 
 protected:
 
-	const FRawMesh& GetRawMesh(const size_t FaceNumber, int32& MeshIdx, int32& LocalFaceNumber) const;
+	const FMeshDescription& GetRawMesh(const size_t FaceNumber, int32& MeshIdx, int32& LocalFaceNumber, const FMeshDescriptionAttributesGetter** OutAttributesGetter) const;
 	
 
 	void ComputeAABB(ProxyLOD::FBBox& InOutBBox);
@@ -370,14 +415,15 @@ protected:
 	ProxyLOD::FBBox   BBox;
 
 	std::vector<size_t>                 PolyOffsetArray;
-	std::vector<const FRawMesh*>        RawMeshArray;
+	std::vector<FMeshDescription*> RawMeshArray;
+	std::vector<FMeshDescriptionAttributesGetter> RawMeshArrayData;
 
 	std::vector<const FMeshMergeData*>  MergeDataArray;
 
 };
 
 /**
-* Class that associates reference geometry (in the form of a FRawMeshArrayAdapter) and a sparse index grid.
+* Class that associates reference geometry (in the form of a FMeshDescriptionArrayAdapter) and a sparse index grid.
 *
 * In standard use, the sparse index grid will hold the Id of the closest poly to the center of each voxel
 * in the voxel.
@@ -401,12 +447,12 @@ public:
 	typedef std::shared_ptr<FClosestPolyField>        Ptr;
 	typedef std::shared_ptr<const FClosestPolyField>  ConstPtr;
 
-	static ConstPtr create(const FRawMeshArrayAdapter& MeshArray, const openvdb::Int32Grid::Ptr SrcPolyIndexGrid)
+	static ConstPtr create(const FMeshDescriptionArrayAdapter& MeshArray, const openvdb::Int32Grid::Ptr SrcPolyIndexGrid)
 	{
 		return ConstPtr(new FClosestPolyField(MeshArray, SrcPolyIndexGrid));
 	}
 
-	FClosestPolyField(const FRawMeshArrayAdapter& MeshArray, const openvdb::Int32Grid::Ptr& SrcPolyIndexGrid); 
+	FClosestPolyField(const FMeshDescriptionArrayAdapter& MeshArray, const openvdb::Int32Grid::Ptr& SrcPolyIndexGrid);
 
 	FClosestPolyField(const FClosestPolyField& other);
 		
@@ -418,7 +464,7 @@ public:
 	class FPolyConstAccessor
 	{
 	public:
-		FPolyConstAccessor(const openvdb::Int32Grid* PolyIndexGrid, const FRawMeshArrayAdapter* MeshArrayAdapter); 
+		FPolyConstAccessor(const openvdb::Int32Grid* PolyIndexGrid, const FMeshDescriptionArrayAdapter* MeshArrayAdapter);
 
 		/** 
 		* Return the closest poly, also let us know if the query failed.
@@ -428,14 +474,14 @@ public:
 		*
 		* @return FRawyPoly  Stuct holding all the data associated with the closest poly in the reference geometry.
 		*/
-		FRawMeshArrayAdapter::FRawPoly  Get(const openvdb::Vec3d& WorldPos, bool& bSuccess) const;
-		FRawMeshArrayAdapter::FRawPoly  Get(const FVector& WorldPos, bool& bSuccess) const
+		FMeshDescriptionArrayAdapter::FRawPoly  Get(const openvdb::Vec3d& WorldPos, bool& bSuccess) const;
+		FMeshDescriptionArrayAdapter::FRawPoly  Get(const FVector& WorldPos, bool& bSuccess) const
 		{
 			return this->Get(openvdb::Vec3d(WorldPos.X, WorldPos.Y, WorldPos.Z), bSuccess);
 		}
 
 	private:
-		const FRawMeshArrayAdapter*       MeshArray;
+		const FMeshDescriptionArrayAdapter*       MeshArray;
 		openvdb::Int32Grid::ConstAccessor CAccessor;
 		const openvdb::math::Transform*   XForm;
 
@@ -457,11 +503,11 @@ public:
 	/**
 	* Access to the underlying reference geometry.
 	*/
-	const FRawMeshArrayAdapter& MeshAdapter() const { return *RawMeshArrayAdapter; }
+	const FMeshDescriptionArrayAdapter& MeshAdapter() const { return *RawMeshArrayAdapter; }
 
 private:
 
-	const FRawMeshArrayAdapter* RawMeshArrayAdapter;
+	const FMeshDescriptionArrayAdapter* RawMeshArrayAdapter;
 	openvdb::Int32Grid::ConstPtr   ClosestPolyGrid;
 };
 
diff --git a/Engine/Plugins/Experimental/ProxyLODPlugin/Source/ProxyLOD/Private/ProxyLODMeshUtilities.cpp b/Engine/Plugins/Experimental/ProxyLODPlugin/Source/ProxyLOD/Private/ProxyLODMeshUtilities.cpp
--- a/Engine/Plugins/Experimental/ProxyLODPlugin/Source/ProxyLOD/Private/ProxyLODMeshUtilities.cpp
+++ b/Engine/Plugins/Experimental/ProxyLODPlugin/Source/ProxyLOD/Private/ProxyLODMeshUtilities.cpp
@@ -12,6 +12,12 @@
 
 #include <DirectXMeshCode/DirectXMesh/DirectXMesh.h>
 
+#include "MeshDescription.h"
+#include "MeshAttributes.h"
+#include "MeshAttributeArray.h"
+
+#include "MeshDescriptionOperations.h"
+
 #include <vector>
 #include <map>
 #include <unordered_map>
@@ -21,22 +27,41 @@
 #ifndef  PROXYLOD_CLOCKWISE_TRIANGLES
 	#define PROXYLOD_CLOCKWISE_TRIANGLES  1
 #endif
-// Compute a tangent space for a FRawMesh 
+// Compute a tangent space for a FMeshDescription 
 
-void ProxyLOD::ComputeTangentSpace( FRawMesh& RawMesh, const bool bRecomputeNormals)
+void ProxyLOD::ComputeTangentSpace(FMeshDescription& RawMesh, const bool bRecomputeNormals)
 {
+	FVertexInstanceArray& VertexInstanceArray = RawMesh.VertexInstances();
+	TVertexInstanceAttributesRef<FVector> Normals = RawMesh.VertexInstanceAttributes().GetAttributesRef<FVector>(MeshAttribute::VertexInstance::Normal);
+	TVertexInstanceAttributesRef<FVector> Tangents = RawMesh.VertexInstanceAttributes().GetAttributesRef<FVector>(MeshAttribute::VertexInstance::Tangent);
+	TVertexInstanceAttributesRef<float> BinormalSigns = RawMesh.VertexInstanceAttributes().GetAttributesRef<float>(MeshAttribute::VertexInstance::BinormalSign);
+
+	// Static meshes always blend normals of overlapping corners.
+	uint32 TangentOptions = FMeshDescriptionOperations::ETangentOptions::BlendOverlappingNormals | FMeshDescriptionOperations::ETangentOptions::IgnoreDegenerateTriangles;
+
+	//Keep the original mesh description NTBs if we do not rebuild the normals or tangents.
+	bool bHasAllNormals = true;
+	bool bHasAllTangents = true;
+	for (const FVertexInstanceID VertexInstanceID : VertexInstanceArray.GetElementIDs())
+	{
+		//Dump the tangents
+		BinormalSigns[VertexInstanceID] = 0.0f;
+		Tangents[VertexInstanceID] = FVector(0.0f);
 
-	// Used to compute the Mikk-Tangent space
-	IMeshUtilities& Utilities = FModuleManager::Get().LoadModuleChecked<IMeshUtilities>("MeshUtilities");
-
-	const bool bRecomputeTangents = true;
-
-	FMeshBuildSettings BuildSettings;
-	BuildSettings.bUseMikkTSpace = true;
-	BuildSettings.bRemoveDegenerates = true; // this is really being used to help identify duplicate 'wedges' 
-
-	Utilities.RecomputeTangentsAndNormalsForRawMesh(bRecomputeTangents, bRecomputeNormals, BuildSettings, RawMesh);
+		if (bRecomputeNormals)
+		{
+			//Dump the normals
+			Normals[VertexInstanceID] = FVector(0.0f);
+		}
+		bHasAllNormals &= !Normals[VertexInstanceID].IsNearlyZero();
+		bHasAllTangents &= !Tangents[VertexInstanceID].IsNearlyZero();
+	}
 
+	if (!bHasAllNormals)
+	{
+		FMeshDescriptionOperations::CreateNormals(RawMesh, (FMeshDescriptionOperations::ETangentOptions)TangentOptions, false);
+	}
+	FMeshDescriptionOperations::CreateMikktTangents(RawMesh, (FMeshDescriptionOperations::ETangentOptions)TangentOptions);
 }
 
 // Calls into the direxXMesh library to compute the per-vertex normal, by default this will weight by area.
@@ -1434,20 +1459,157 @@ void ProxyLOD::AddNormals(TAOSMesh<FPositionOnlyVertex>& InOutMesh)
 		return TestCount;
 	}
 
+/**
+* Attempt to correct the collapsed walls.
+* NB: The kDOP tree is already built, using the same mesh.
+*
+* @param Indices  - mesh conectivity
+* @param Positions - vertex locations: maybe changed by this function.
+* @param VoxelSize - length scale used in heuristic that determins how far to move vertices.
+*/
+int32 CorrectCollapsedWalls(const ProxyLOD::FkDOPTree& kDOPTree,
+	FMeshDescription& MeshDescription,
+	const float VoxelSize)
+{
+	typedef uint32 EdgeIdType;
+	typedef uint32 FaceIdType;
+	TVertexAttributesRef<FVector> VertexPositions = MeshDescription.VertexAttributes().GetAttributesRef<FVector>(MeshAttribute::Vertex::Position);
+
+	ProxyLOD::FUnitTransformDataProvider kDOPDataProvider(kDOPTree);
+
+	// Number of triangle in our mesh
+	int32 NumTriangle = 0;
+	for (const FPolygonID& PolygonID : MeshDescription.Polygons().GetElementIDs())
+	{
+		NumTriangle += MeshDescription.GetPolygonTriangles(PolygonID).Num();
+	}
+
+	// This will hold the intersecting faces for each edge.
+
+	std::unordered_map<FaceIdType, std::vector<FaceIdType>> IntersectionListMap;
+
+	//const auto* Indices = IndexArray.GetData();
+	//auto* Pos = PositionArray.GetData();
+
+	// loop over the polys and collect the names of the faces that intersect.
+	auto GetFace = [&MeshDescription, &VertexPositions](int32 FaceIdx, FVector(&Verts)[3])
+	{
+		const FVertexID Idx[3] = { MeshDescription.GetVertexInstanceVertex(FVertexInstanceID(3 * FaceIdx)),
+								MeshDescription.GetVertexInstanceVertex(FVertexInstanceID(3 * FaceIdx + 1)),
+								MeshDescription.GetVertexInstanceVertex(FVertexInstanceID(3 * FaceIdx + 2)) };
+
+		Verts[0] = VertexPositions[Idx[0]];
+		Verts[1] = VertexPositions[Idx[1]];
+		Verts[2] = VertexPositions[Idx[2]];
+	};
+
+	auto GetFaceNormal = [&GetFace](int32 FaceIdx)->FVector
+	{
+		FVector Verts[3];
+		GetFace(FaceIdx, Verts);
+
+		return ComputeNormal(Verts);
+	};
+
+	int32 TestCount = 0;
+	for (int32 FaceIdx = 0; FaceIdx < NumTriangle; ++FaceIdx)
+	{
+		FVector Verts[3];
+		GetFace(FaceIdx, Verts);
+
+		const FVector FaceNormal = ComputeNormal(Verts);
+
+		// loop over these three edges.
+		for (int32 j = 0; j < 3; ++j)
+		{
+			int32 sV = j;
+			int32 eV = (j + 1) % 3;
+
+			FkHitResult kDOPResult;
+
+			TkDOPLineCollisionCheck<const ProxyLOD::FUnitTransformDataProvider, uint32>  EdgeRay(Verts[sV], Verts[eV], true, kDOPDataProvider, &kDOPResult);
+
+			bool bHit = kDOPTree.LineCheck(EdgeRay);
+
+			if (bHit)
+			{
+				// Triangle we hit
+				int32 HitTriId = kDOPResult.Item;
+
+				// Don't count a hit against myself
+				if (HitTriId == FaceIdx)
+				{
+					continue;
+				}
+
+				// Make sure the hit wasn't just one of the verts.
+				if (kDOPResult.Time > 0.999 || kDOPResult.Time < 0.001)
+				{
+					continue;
+				}
+
+				// We only care about faces pointing in opposing directions
+				const FVector HitFaceNormal = GetFaceNormal(HitTriId);
+				if (FVector::DotProduct(FaceNormal, HitFaceNormal) > -0.94f) // not in 160 to 200 degrees
+				{
+					continue;
+				}
+
+				TestCount++;
+
+				auto Search = IntersectionListMap.find(FaceIdx);
+				if (Search != IntersectionListMap.end())
+				{
+					auto& FaceList = Search->second;
+					if (std::find(FaceList.begin(), FaceList.end(), HitTriId) == FaceList.end())
+					{
+						FaceList.push_back(HitTriId);
+					}
+				}
+				else
+				{
+					std::vector<FaceIdType> FaceList;
+					FaceList.push_back(HitTriId);
+					IntersectionListMap[FaceIdx] = FaceList;
+				}
+			}
+		}
+	}
+
+	//For each triangle that collides, push it a small fixed distance in the normal direction.
+	{
+		for (auto ListMapIter = IntersectionListMap.begin(); ListMapIter != IntersectionListMap.end(); ++ListMapIter)
+		{
+			int32 FaceIdx = ListMapIter->first;
+			const FVector TriNormal = GetFaceNormal(FaceIdx);
+
+			// Scale by a small amount
+
+			const FVector NormDisplacement = TriNormal * (VoxelSize / 7.f);
+
+			const FVertexID Idx[3] = {	MeshDescription.GetVertexInstanceVertex(FVertexInstanceID(3 * FaceIdx)),
+										MeshDescription.GetVertexInstanceVertex(FVertexInstanceID(3 * FaceIdx + 1)),
+										MeshDescription.GetVertexInstanceVertex(FVertexInstanceID(3 * FaceIdx + 2)) };
 
-int32 ProxyLOD::CorrectCollapsedWalls( FRawMesh& InOutMesh, 
-	                                   const float VoxelSize )
+			VertexPositions[Idx[0]] += NormDisplacement;
+			VertexPositions[Idx[1]] += NormDisplacement;
+			VertexPositions[Idx[2]] += NormDisplacement;
+		}
+
+	}
+	return TestCount;
+}
+
+int32 ProxyLOD::CorrectCollapsedWalls(FMeshDescription& InOutMeshDescription, const float VoxelSize)
 {
 	// Build an acceleration structure
 
 	FkDOPTree kDOPTree;
-	BuildkDOPTree(InOutMesh, kDOPTree);
+	BuildkDOPTree(InOutMeshDescription, kDOPTree);
 
-	return CorrectCollapsedWalls(kDOPTree, InOutMesh.WedgeIndices, InOutMesh.VertexPositions, VoxelSize);
-	
+	return CorrectCollapsedWalls(kDOPTree, InOutMeshDescription, VoxelSize);
 }
 
-
 int32 ProxyLOD::CorrectCollapsedWalls(FVertexDataMesh& InOutMesh,
 	const float VoxelSize)
 {
@@ -1508,7 +1670,7 @@ void ProxyLOD::TestUniqueVertexes(const FAOSMesh& InMesh)
 }
 
 
-void ProxyLOD::ColorPartitions(FRawMesh& InOutRawMesh, const std::vector<uint32>& partitionResults)
+void ProxyLOD::ColorPartitions(FMeshDescription& InOutRawMesh, const std::vector<uint32>& partitionResults)
 {
 
 	// testing - coloring the simplified mesh by the partitions generated by uvatlas
@@ -1516,13 +1678,33 @@ void ProxyLOD::ColorPartitions(FRawMesh& InOutRawMesh, const std::vector<uint32>
 		FColor(153, 102, 0), FColor(249, 129, 162), FColor(29, 143, 177), FColor(118, 42, 145),
 		FColor(255, 121, 75), FColor(102, 204, 51), FColor(153, 153, 255), FColor(255, 255, 255) };
 
+	TVertexInstanceAttributesRef<FVector4> VertexInstanceColors = InOutRawMesh.VertexInstanceAttributes().GetAttributesRef<FVector4>(MeshAttribute::VertexInstance::Color);
+
+	//Remap the vertex instance
+	int32 TriangleIndex = 0;
+	TMap<uint32, FVertexInstanceID> WedgeIndexToVertexInstanceID;
+	WedgeIndexToVertexInstanceID.Reserve(InOutRawMesh.VertexInstances().Num());
+	for (const FPolygonID& PolygonID : InOutRawMesh.Polygons().GetElementIDs())
+	{
+		const FMeshPolygon& Polygon = InOutRawMesh.GetPolygon(PolygonID);
+		for (const FMeshTriangle& Triangle : Polygon.Triangles)
+		{
+			for (int32 Corner = 0; Corner < 3; ++Corner)
+			{
+				WedgeIndexToVertexInstanceID.Add((TriangleIndex * 3) + Corner, Triangle.GetVertexInstanceID(Corner));
+			}
+			TriangleIndex++;
+		}
+	}
+
 	for (int i = 0; i < partitionResults.size(); ++i)
 	{
 		uint32 PId = partitionResults[i];
-		uint32 WedgeId = i * 3;
-		InOutRawMesh.WedgeColors[WedgeId + 0] = Range[PId % 13];
-		InOutRawMesh.WedgeColors[WedgeId + 1] = Range[PId % 13];
-		InOutRawMesh.WedgeColors[WedgeId + 2] = Range[PId % 13];
+		for (int32 Corner = 0; Corner < 3; ++Corner)
+		{
+			FVertexInstanceID VertexInstanceID = WedgeIndexToVertexInstanceID[(i * 3) + Corner];
+			VertexInstanceColors[VertexInstanceID] = FLinearColor(Range[PId % 13]);
+		}
 	}
 }
 
@@ -1548,27 +1730,30 @@ void ProxyLOD::ColorPartitions(FVertexDataMesh& InOutMesh, const std::vector<uin
 }
 
 
-void ProxyLOD::AddWedgeColors(FRawMesh& RawMesh)
+void ProxyLOD::AddWedgeColors(FMeshDescription& RawMesh)
 {
 
 	FColor ColorRange[13] = { FColor(255, 0, 0), FColor(0, 255, 0), FColor(0, 0, 255), FColor(255, 255, 0), FColor(0, 255, 255),
 		FColor(153, 102, 0), FColor(249, 129, 162), FColor(29, 143, 177), FColor(118, 42, 145),
 		FColor(255, 121, 75), FColor(102, 204, 51), FColor(153, 153, 255), FColor(255, 255, 255) };
 
-	int32 NumIndices = RawMesh.WedgeIndices.Num();
-	auto& WedgeColors = RawMesh.WedgeColors;
+	TVertexInstanceAttributesRef<FVector4> VertexInstanceColors = RawMesh.VertexInstanceAttributes().GetAttributesRef<FVector4>(MeshAttribute::VertexInstance::Color);
 
-	ResizeArray(WedgeColors, NumIndices);
-
-	ProxyLOD::Parallel_For(ProxyLOD::FIntRange(0, NumIndices),
-		[&ColorRange, &WedgeColors](const ProxyLOD::FIntRange& Range)
+	//Recolor the vertex instances
+	int32 TriangleIndex = 0;
+	for (const FPolygonID& PolygonID : RawMesh.Polygons().GetElementIDs())
 	{
-		for (int32 i = Range.begin(), I = Range.end(); i < I; ++i)
+		const FMeshPolygon& Polygon = RawMesh.GetPolygon(PolygonID);
+		for (const FMeshTriangle& Triangle : Polygon.Triangles)
 		{
-			WedgeColors[i] = ColorRange[i % 13];
-		}
-	});
+			for (int32 Corner = 0; Corner < 3; ++Corner)
+			{
 
+				VertexInstanceColors[Triangle.GetVertexInstanceID(Corner)] = FLinearColor(ColorRange[((TriangleIndex*3) + Corner) % 13]);
+			}
+			TriangleIndex++;
+		}
+	}
 }
 
 
@@ -1657,7 +1842,7 @@ void ProxyLOD::AddNormals(FAOSMesh& InOutMesh)
 // Unused
 #if 0
 // Computes the face normals and assigns them to the wedges TangentZ
-static void ComputeRawMeshNormals(FRawMesh& InOutMesh)
+static void ComputeRawMeshNormals(FMeshDescription& InOutMesh)
 {
 	const int32 NumFaces = InOutMesh.WedgeIndices.Num() / 3;
 
diff --git a/Engine/Plugins/Experimental/ProxyLODPlugin/Source/ProxyLOD/Private/ProxyLODMeshUtilities.h b/Engine/Plugins/Experimental/ProxyLODPlugin/Source/ProxyLOD/Private/ProxyLODMeshUtilities.h
--- a/Engine/Plugins/Experimental/ProxyLODPlugin/Source/ProxyLOD/Private/ProxyLODMeshUtilities.h
+++ b/Engine/Plugins/Experimental/ProxyLODPlugin/Source/ProxyLOD/Private/ProxyLODMeshUtilities.h
@@ -32,13 +32,13 @@ namespace ProxyLOD
 	* Uses UE4 code in the Mesh Utilities module to compute the tangent space
 	* The computation is a Mikk-T(angent) space version of tangent space.  
 	*
-	* NB: This updates the FRawMesh::TangentX and FRawMesh::TangentY and optionally TangentZ.
+	* NB: This updates the FMeshDescription::TangentX and FMeshDescription::TangentY and optionally TangentZ.
 	*     This can result in a per-wedge tangent space if the normals are split.
 	* 
-	* @param InOutMesh          FRawMesh which will be updated by this function - adding a tangent space
+	* @param InOutMesh          FMeshDescription which will be updated by this function - adding a tangent space
 	* @param bRecomputeNormals  Determine if the normal should be recomputed as well (defualt - no).
 	*/
-	void ComputeTangentSpace( FRawMesh& InOutMesh, 
+	void ComputeTangentSpace(FMeshDescription& InOutMesh,
 		                      const bool bRecomputeNormals = false );
 	/**
 	* Calls into the direxXMesh library to compute the per-vertex tangent and bitangent, optionally recomputes the normal.
@@ -141,17 +141,17 @@ namespace ProxyLOD
 	*
 	* @return  Number of detected inter-penetrations of opposing faced polys.
 	*/
-	int32 CorrectCollapsedWalls( FRawMesh& InOutMesh, 
-		                         const float LengthScale );
+	int32 CorrectCollapsedWalls( FMeshDescription& InOutMeshDescription,
+								 const float LengthScale);
 	int32 CorrectCollapsedWalls( FVertexDataMesh& InOutMesh,
 		                         const float LengthScale );
 	
 #if 0
 	// Can be used to transfer things like vertex colors from source geometry to output.
 
-	void TransferMeshAttributes( const FRawMeshArrayAdapter& RawMeshArrayAdapter,
+	void TransferMeshAttributes( const FMeshDescriptionArrayAdapter& RawMeshArrayAdapter,
 		                         const openvdb::Int32Grid& SrcPolyIndexGrid,
-		                         FRawMesh& InOutMesh );
+								 FMeshDescription& InOutMesh );
 #endif 
 	// --------------- debugging ---------------------------------------------------------------
 
@@ -174,7 +174,7 @@ namespace ProxyLOD
 	* @param InOutMesh       Mesh to be colored.
 	* @param PartitionArray  Maps each face to a partition.
 	*/
-	void ColorPartitions( FRawMesh& InOutMesh,     const std::vector<uint32>& PartitionArray );
+	void ColorPartitions(FMeshDescription& InOutMesh,     const std::vector<uint32>& PartitionArray );
 	void ColorPartitions( FVertexDataMesh& InMesh, const std::vector<uint32>& PartitionArray );
 
 	
@@ -183,7 +183,7 @@ namespace ProxyLOD
 	* 
 	* @param InOutMesh  Mesh to be colored.
 	*/
-	void AddWedgeColors( FRawMesh& InOutMesh );
+	void AddWedgeColors(FMeshDescription& InOutMesh );
 	
 	/**
 	* Add a default tangent and bi tangent  aligned with the coordinate axes :(1, 0, 0), (0, 1, 0) to a mesh.
diff --git a/Engine/Plugins/Experimental/ProxyLODPlugin/Source/ProxyLOD/Private/ProxyLODParallelSimplifier.cpp b/Engine/Plugins/Experimental/ProxyLODPlugin/Source/ProxyLOD/Private/ProxyLODParallelSimplifier.cpp
--- a/Engine/Plugins/Experimental/ProxyLODPlugin/Source/ProxyLOD/Private/ProxyLODParallelSimplifier.cpp
+++ b/Engine/Plugins/Experimental/ProxyLODPlugin/Source/ProxyLOD/Private/ProxyLODParallelSimplifier.cpp
@@ -71,7 +71,7 @@ void ProxyLOD::PartitionAndSimplifyMesh(const FBBox& SrcBBox, const TArray<int32
 
 void ProxyLOD::ParallelSimplifyMesh(const FClosestPolyField& SrcReference, const float MinFractionToRatain, const float MaxFeatureCost, FAOSMesh& InOutMesh)
 {
-	const FRawMeshArrayAdapter& SrcMeshAdapter = SrcReference.MeshAdapter();
+	const FMeshDescriptionArrayAdapter& SrcMeshAdapter = SrcReference.MeshAdapter();
 
 	const size_t SrcTriNum = SrcMeshAdapter.polygonCount();
 
diff --git a/Engine/Plugins/Experimental/ProxyLODPlugin/Source/ProxyLOD/Private/ProxyLODPlugin.cpp b/Engine/Plugins/Experimental/ProxyLODPlugin/Source/ProxyLOD/Private/ProxyLODPlugin.cpp
--- a/Engine/Plugins/Experimental/ProxyLODPlugin/Source/ProxyLOD/Private/ProxyLODPlugin.cpp
+++ b/Engine/Plugins/Experimental/ProxyLODPlugin/Source/ProxyLOD/Private/ProxyLODPlugin.cpp
@@ -5,7 +5,6 @@
 #include "Modules/ModuleManager.h"
 #include "Features/IModularFeatures.h"
 #include "IProxyLODPlugin.h"
-#include "RawMesh.h"
 #include "MeshMergeData.h"
 #include "Engine/MeshMerging.h"
 #include "MaterialUtilities.h" // for FFlattenMaterial 
@@ -36,6 +35,11 @@
 #include "ProxyLODkDOPInterface.h"
 #include "ProxyLODThreadedWrappers.h"
 
+#include "MeshDescription.h"
+#include "MeshAttributes.h"
+#include "MeshAttributeArray.h"
+#include "MeshDescriptionOperations.h"
+
 
 
 #define PROXYLOD_USE_TEST_CVARS 1
@@ -419,7 +423,7 @@ static void SimplifyMesh( const FClosestPolyField& SrcGeometryPolyField,
 	//SCOPE_LOG_TIME(TEXT("UE4_ProxyLOD_Simplifier"), nullptr);
 
 	// Compute some of the metrics that relate the desired resolution to simplifier parameters.
-	const FRawMeshArrayAdapter& SrcGeometryAdapter = SrcGeometryPolyField.MeshAdapter();
+	const FMeshDescriptionArrayAdapter& SrcGeometryAdapter = SrcGeometryPolyField.MeshAdapter();
 	const ProxyLOD::FBBox& SrcBBox = SrcGeometryAdapter.GetBBox();
 	const float MaxSide = SrcBBox.extents().length();
 
@@ -578,21 +582,20 @@ void FVoxelizeMeshMerging::ProxyLOD(const FMeshMergeDataArray& InData, const FMe
 	// NB: These will be the product of this function and 
 	//     will be captured by the CompleteDelegate. 
 
-	FRawMesh OutRawMesh;
+	FMeshDescription OutRawMesh;
+	UStaticMesh::RegisterMeshAttributes(OutRawMesh);
 
 	FFlattenMaterial OutMaterial         = FMaterialUtilities::CreateFlattenMaterialWithSettings(InProxySettings.MaterialSettings);
 	const FColor UnresolvedGeometryColor = InProxySettings.UnresolvedGeometryColor;
 
-
-
 	bool bProxyGenerationSuccess = true;
 	// Compute the simplified mesh and related materials.
 	{
 	
 		// Create an adapter to make the data appear as a single mesh as required by the voxelization code.
 
-		FRawMeshArrayAdapter SrcGeometryAdapter(InGeometry);
-		FRawMeshArrayAdapter ClippingGeometryAdapter(InClippingGeometry); 
+		FMeshDescriptionArrayAdapter SrcGeometryAdapter(InGeometry);
+		FMeshDescriptionArrayAdapter ClippingGeometryAdapter(InClippingGeometry);
 
 		{
 			const auto& BBox = SrcGeometryAdapter.GetBBox();
@@ -838,7 +841,7 @@ void FVoxelizeMeshMerging::ProxyLOD(const FMeshMergeDataArray& InData, const FMe
 					}
 				}
 			);
-
+			SrcGeometryAdapter.UpdateMaterialsID();
 		}
 		else
 		{
@@ -887,7 +890,7 @@ void FVoxelizeMeshMerging::ProxyLOD(const FMeshMergeDataArray& InData, const FMe
 
 
 			// --- Compute the tangent space on the simplified mesh ---
-			// Convert to FRawMesh because it has the per-index attribute structure needed for the tangent space.
+			// Convert to FMeshDescription because it has the per-index attribute structure needed for the tangent space.
 			// Generate the tangent space, but retain our current normals.
 
 			MapTextureAtlasAndAddTangentSpaceTaskGroup.RunAndWait(
@@ -1003,19 +1006,25 @@ void FVoxelizeMeshMerging::ProxyLOD(const FMeshMergeDataArray& InData, const FMe
 	{
 
 		// Revert the smoothing group to a default.
+		TEdgeAttributesRef<bool> EdgeHardnesses = OutRawMesh.EdgeAttributes().GetAttributesRef<bool>(MeshAttribute::Edge::IsHard);
 
-		for (int32 i = 0; i < OutRawMesh.FaceSmoothingMasks.Num(); ++i)
+		for (const FEdgeID& EdgeID : OutRawMesh.Edges().GetElementIDs())
 		{
-			OutRawMesh.FaceSmoothingMasks[i] = 1;
+			EdgeHardnesses[EdgeID] = false;
 		}
 
-		// Debuging
-
-		checkSlow(OutRawMesh.IsValid());
 
 		// NB: FProxyGenerationProcessor::ProxyGenerationComplete
-
-		CompleteDelegate.ExecuteIfBound(OutRawMesh, OutMaterial, InJobGUID);
+		FRawMesh ConvertedRawMesh;
+		TMap<FName, int32> MaterialMap;
+		TPolygonGroupAttributesConstRef<FName> PolygonGroupNames = OutRawMesh.PolygonGroupAttributes().GetAttributesRef<FName>(MeshAttribute::PolygonGroup::ImportedMaterialSlotName);
+		int32 MaterialIndex = 0;
+		for (const FPolygonGroupID PolygonGroupID : OutRawMesh.PolygonGroups().GetElementIDs())
+		{
+			MaterialMap.Add(PolygonGroupNames[PolygonGroupID], MaterialIndex++);
+		}
+		FMeshDescriptionOperations::ConvertToRawMesh(OutRawMesh, ConvertedRawMesh, MaterialMap);
+		CompleteDelegate.ExecuteIfBound(ConvertedRawMesh, OutMaterial, InJobGUID);
 	}
 	else
 	{
diff --git a/Engine/Plugins/Experimental/ProxyLODPlugin/Source/ProxyLOD/Private/ProxyLODRasterizer.cpp b/Engine/Plugins/Experimental/ProxyLODPlugin/Source/ProxyLOD/Private/ProxyLODRasterizer.cpp
--- a/Engine/Plugins/Experimental/ProxyLODPlugin/Source/ProxyLOD/Private/ProxyLODRasterizer.cpp
+++ b/Engine/Plugins/Experimental/ProxyLODPlugin/Source/ProxyLOD/Private/ProxyLODRasterizer.cpp
@@ -212,12 +212,14 @@ namespace ProxyLOD
 		return RasterGridPtr;
 	}
 
-	static void InterpolateWedgeColors(const FRawMesh& RawMesh, const ProxyLOD::FRasterGrid& DstUVGrid, FFlattenMaterial& OutMaterial)
+	static void InterpolateWedgeColors(const FMeshDescription& RawMesh, const ProxyLOD::FRasterGrid& DstUVGrid, FFlattenMaterial& OutMaterial)
 	{
 		TArray<FColor>& ColorBuffer = OutMaterial.GetPropertySamples(EFlattenMaterialProperties::Diffuse);
 		FIntPoint Size = OutMaterial.GetPropertySize(EFlattenMaterialProperties::Diffuse);
 		ResizeArray(ColorBuffer, Size.X * Size.Y);
 
+		TVertexInstanceAttributesConstRef<FVector4> VertexInstanceColors = RawMesh.VertexInstanceAttributes().GetAttributesRef<FVector4>(MeshAttribute::VertexInstance::Color);
+
 		for (int j = 0; j < Size.Y; ++j)
 		{
 			for (int i = 0; i < Size.X; ++i)
@@ -229,11 +231,11 @@ namespace ProxyLOD
 				{
 					const auto& BarycentericCoords = Data.BarycentricCoords;
 
-					const int32 Indexes[3] = { TriangleId * 3, TriangleId * 3 + 1, TriangleId * 3 + 2 };
+					const FVertexInstanceID Indexes[3] = { FVertexInstanceID(TriangleId * 3), FVertexInstanceID(TriangleId * 3 + 1), FVertexInstanceID(TriangleId * 3 + 2) };
 					FLinearColor ColorSamples[3];
-					ColorSamples[0] = FLinearColor(RawMesh.WedgeColors[Indexes[0]]);
-					ColorSamples[1] = FLinearColor(RawMesh.WedgeColors[Indexes[1]]);
-					ColorSamples[2] = FLinearColor(RawMesh.WedgeColors[Indexes[2]]);
+					ColorSamples[0] = FLinearColor(VertexInstanceColors[Indexes[0]]);
+					ColorSamples[1] = FLinearColor(VertexInstanceColors[Indexes[1]]);
+					ColorSamples[2] = FLinearColor(VertexInstanceColors[Indexes[2]]);
 					FLinearColor AvgColor = ProxyLOD::InterpolateVertexData(BarycentericCoords, ColorSamples);
 
 					ColorBuffer[i + j * Size.X] = AvgColor.ToFColor(true);
@@ -247,7 +249,7 @@ namespace ProxyLOD
 
 	};
 
-	void DebugVertexAndTextureColors(FRawMesh& RawMesh, const ProxyLOD::FRasterGrid& DstUVGrid, FFlattenMaterial& OutMaterial)
+	void DebugVertexAndTextureColors(FMeshDescription& RawMesh, const ProxyLOD::FRasterGrid& DstUVGrid, FFlattenMaterial& OutMaterial)
 	{
 		// Testing the Barycentric interpolation.
 
diff --git a/Engine/Plugins/Experimental/ProxyLODPlugin/Source/ProxyLOD/Private/ProxyLODRasterizer.h b/Engine/Plugins/Experimental/ProxyLODPlugin/Source/ProxyLOD/Private/ProxyLODRasterizer.h
--- a/Engine/Plugins/Experimental/ProxyLODPlugin/Source/ProxyLOD/Private/ProxyLODRasterizer.h
+++ b/Engine/Plugins/Experimental/ProxyLODPlugin/Source/ProxyLOD/Private/ProxyLODRasterizer.h
@@ -105,7 +105,7 @@ namespace ProxyLOD
 	*  @param InUVGrid      Rasterization of the UV triangles of the InOutGrid.
 	*  @param OutMaterial   Material that, on return has the Barycenteric-averaged vertex colors. 
 	*/
-	void DebugVertexAndTextureColors(FRawMesh& InOutRawMesh, const ProxyLOD::FRasterGrid& InUVGrid, FFlattenMaterial& OutMaterial);
+	void DebugVertexAndTextureColors(FMeshDescription& InOutRawMesh, const ProxyLOD::FRasterGrid& InUVGrid, FFlattenMaterial& OutMaterial);
 	
 	/**
 	* Two dimensional grid access to a TArray holding FLinearColor.
diff --git a/Engine/Plugins/Experimental/ProxyLODPlugin/Source/ProxyLOD/Private/ProxyLODkDOPInterface.h b/Engine/Plugins/Experimental/ProxyLODPlugin/Source/ProxyLOD/Private/ProxyLODkDOPInterface.h
--- a/Engine/Plugins/Experimental/ProxyLODPlugin/Source/ProxyLOD/Private/ProxyLODkDOPInterface.h
+++ b/Engine/Plugins/Experimental/ProxyLODPlugin/Source/ProxyLOD/Private/ProxyLODkDOPInterface.h
@@ -7,8 +7,7 @@
 #include "ProxyLODMeshTypes.h"
 #include "ProxyLODThreadedWrappers.h"
 
-#include "RawMesh.h"
-
+class FMeshDescription;
 
 namespace ProxyLOD
 {
@@ -76,7 +75,7 @@ namespace ProxyLOD
 	* @param InMesh     Input triangle mesh
 	* @param kDOPTree   The resulting acceleration structure.
 	*/
-	void BuildkDOPTree( const FRawMesh& InMesh, FkDOPTree& kDOPTree);
+	void BuildkDOPTree( const FMeshDescription& MeshDescription, FkDOPTree& kDOPTree);
 	void BuildkDOPTree( const FVertexDataMesh& InMesh, FkDOPTree& kDOPTree );
-	void BuildkDOPTree( const FRawMeshArrayAdapter& InMesh, FkDOPTree& kDOPTree);
+	void BuildkDOPTree( const FMeshDescriptionArrayAdapter& InMesh, FkDOPTree& kDOPTree);
 }
diff --git a/Engine/Plugins/Experimental/ProxyLODPlugin/Source/ProxyLOD/Private/ProxyLODkDOPUtils.cpp b/Engine/Plugins/Experimental/ProxyLODPlugin/Source/ProxyLOD/Private/ProxyLODkDOPUtils.cpp
--- a/Engine/Plugins/Experimental/ProxyLODPlugin/Source/ProxyLOD/Private/ProxyLODkDOPUtils.cpp
+++ b/Engine/Plugins/Experimental/ProxyLODPlugin/Source/ProxyLOD/Private/ProxyLODkDOPUtils.cpp
@@ -5,9 +5,13 @@
 #include "ProxyLODMeshTypes.h"
 #include "ProxyLODThreadedWrappers.h"
 
+#include "MeshAttributes.h"
+#include "MeshAttributeArray.h"
+#include "MeshDescription.h"
+
 // Utils for building a kdop tree from different mesh types.
 
-void ProxyLOD::BuildkDOPTree(const FRawMeshArrayAdapter& SrcGeometry, ProxyLOD::FkDOPTree& kDOPTree)
+void ProxyLOD::BuildkDOPTree(const FMeshDescriptionArrayAdapter& SrcGeometry, ProxyLOD::FkDOPTree& kDOPTree)
 {
 
 	const auto NumSrcPoly = SrcGeometry.polygonCount();
@@ -34,10 +38,15 @@ void ProxyLOD::BuildkDOPTree(const FRawMeshArrayAdapter& SrcGeometry, ProxyLOD::
 
 }
 
-void ProxyLOD::BuildkDOPTree(const FRawMesh& SrcRawMesh, ProxyLOD::FkDOPTree& kDOPTree)
+void ProxyLOD::BuildkDOPTree(const FMeshDescription& MeshDescription, FkDOPTree& kDOPTree)
 {
+	TVertexAttributesConstRef<FVector> VertexPositions = MeshDescription.VertexAttributes().GetAttributesRef<FVector>(MeshAttribute::Vertex::Position);
 
-	const auto NumSrcPoly = SrcRawMesh.WedgeIndices.Num() / 3;
+	uint32 NumSrcPoly = 0;
+	for (const FPolygonID& PolygonID : MeshDescription.Polygons().GetElementIDs())
+	{
+		NumSrcPoly += MeshDescription.GetPolygonTriangles(PolygonID).Num();
+	}
 
 	TArray<FkDOPBuildTriangle> BuildTriangleArray;
 
@@ -45,15 +54,16 @@ void ProxyLOD::BuildkDOPTree(const FRawMesh& SrcRawMesh, ProxyLOD::FkDOPTree& kD
 	ResizeArray(BuildTriangleArray, NumSrcPoly);
 
 	ProxyLOD::Parallel_For(ProxyLOD::FUIntRange(0, NumSrcPoly),
-		[&BuildTriangleArray, &SrcRawMesh](const ProxyLOD::FUIntRange& Range)
+		[&BuildTriangleArray, &MeshDescription, &VertexPositions](const ProxyLOD::FUIntRange& Range)
 	{
 		FkDOPBuildTriangle* BuildTriangles = BuildTriangleArray.GetData();
-		const uint32* Idxs = SrcRawMesh.WedgeIndices.GetData();
-		const FVector* Positions = SrcRawMesh.VertexPositions.GetData();
 
 		for (uint32 r = Range.begin(), R = Range.end(); r < R; ++r)
 		{
-			FVector Pos[3] = { Positions[Idxs[3 * r]],  Positions[Idxs[3 * r + 1]],  Positions[Idxs[3 * r + 2]] };
+
+			FVector Pos[3] = {	VertexPositions[MeshDescription.GetVertexInstanceVertex(FVertexInstanceID(3 * r))],
+								VertexPositions[MeshDescription.GetVertexInstanceVertex(FVertexInstanceID(3 * r + 1))],
+								VertexPositions[MeshDescription.GetVertexInstanceVertex(FVertexInstanceID(3 * r + 2))] };
 			BuildTriangles[r] = FkDOPBuildTriangle(r, Pos[0], Pos[1], Pos[2]);
 		}
 
diff --git a/Engine/Plugins/Experimental/ProxyLODPlugin/Source/ProxyLOD/ProxyLODMeshReduction.Build.cs b/Engine/Plugins/Experimental/ProxyLODPlugin/Source/ProxyLOD/ProxyLODMeshReduction.Build.cs
--- a/Engine/Plugins/Experimental/ProxyLODPlugin/Source/ProxyLOD/ProxyLODMeshReduction.Build.cs
+++ b/Engine/Plugins/Experimental/ProxyLODPlugin/Source/ProxyLOD/ProxyLODMeshReduction.Build.cs
@@ -37,6 +37,7 @@ namespace UnrealBuildTool.Rules
                     "InputCore",
                     "UnrealEd",
                     "RawMesh",
+                    "MeshDescription",
                     "MeshUtilities",
                     "MaterialUtilities",
                     "PropertyEditor",
@@ -56,7 +57,8 @@ namespace UnrealBuildTool.Rules
                     "Core",
                     "CoreUObject",
                     "Engine",
-                    "UnrealEd"
+                    "UnrealEd",
+                    "MeshDescriptionOperations"
 					// ... add private dependencies that you statically link with here ...
 				}
 				);
diff --git a/Engine/Source/Developer/MeshDescriptionOperations/Private/MeshDescriptionOperations.cpp b/Engine/Source/Developer/MeshDescriptionOperations/Private/MeshDescriptionOperations.cpp
--- a/Engine/Source/Developer/MeshDescriptionOperations/Private/MeshDescriptionOperations.cpp
+++ b/Engine/Source/Developer/MeshDescriptionOperations/Private/MeshDescriptionOperations.cpp
@@ -594,6 +594,51 @@ void FMeshDescriptionOperations::ConvertFromRawMesh(const FRawMesh& SourceRawMes
 
 //////////////////////////////////////////////////////////////////////////
 // Normals tangents and Bi-normals
+void FMeshDescriptionOperations::RecomputeNormalsAndTangentsIfNeeded(FMeshDescription& MeshDescription, ETangentOptions TangentOptions, bool bUseMikkTSpace, bool bForceRecomputeNormals, bool bForceRecomputeTangents)
+{
+	bool bRecomputeNormals = bForceRecomputeNormals;
+	bool bRecomputeTangents = bForceRecomputeTangents;
+	TVertexInstanceAttributesRef<FVector> VertexInstanceNormals = MeshDescription.VertexInstanceAttributes().GetAttributesRef<FVector>(MeshAttribute::VertexInstance::Normal);
+	TVertexInstanceAttributesRef<FVector> VertexInstanceTangents = MeshDescription.VertexInstanceAttributes().GetAttributesRef<FVector>(MeshAttribute::VertexInstance::Tangent);
+	if (!bRecomputeNormals || !bRecomputeTangents)
+	{
+		for (const FVertexInstanceID& VertexInstanceID : MeshDescription.VertexInstances().GetElementIDs())
+		{
+			bRecomputeNormals |= (VertexInstanceNormals[VertexInstanceID].IsNearlyZero() || VertexInstanceNormals[VertexInstanceID].ContainsNaN());
+			bRecomputeTangents |= (VertexInstanceTangents[VertexInstanceID].IsNearlyZero() || VertexInstanceTangents[VertexInstanceID].ContainsNaN());
+			if (bRecomputeNormals && bRecomputeTangents)
+			{
+				break;
+			}
+		}
+	}
+
+	if (bRecomputeNormals || bRecomputeTangents)
+	{
+		//Zero out all value that need to be recompute
+		for (const FVertexInstanceID& VertexInstanceID : MeshDescription.VertexInstances().GetElementIDs())
+		{
+			if (bRecomputeNormals)
+			{
+				VertexInstanceNormals[VertexInstanceID] = FVector::ZeroVector;
+			}
+			if (bRecomputeTangents)
+			{
+				VertexInstanceTangents[VertexInstanceID] = FVector::ZeroVector;
+			}
+		}
+
+
+		if (bRecomputeNormals)
+		{
+			FMeshDescriptionOperations::CreateNormals(MeshDescription, TangentOptions, bUseMikkTSpace ? false : bRecomputeTangents);
+		}
+		if (bUseMikkTSpace && bRecomputeTangents)
+		{
+			FMeshDescriptionOperations::CreateMikktTangents(MeshDescription, TangentOptions);
+		}
+	}
+}
 
 void FMeshDescriptionOperations::CreatePolygonNTB(FMeshDescription& MeshDescription, float ComparisonThreshold)
 {
@@ -1264,4 +1309,47 @@ bool FMeshDescriptionOperations::RemoveUVChannel(FMeshDescription& MeshDescripti
 	return true;
 }
 
+void FMeshDescriptionOperations::RemapPolygonGroups(FMeshDescription& MeshDescription, TMap<FPolygonGroupID, FPolygonGroupID>& Remap)
+{
+	TPolygonGroupAttributesRef<FName> PolygonGroupNames = MeshDescription.PolygonGroupAttributes().GetAttributesRef<FName>(MeshAttribute::PolygonGroup::ImportedMaterialSlotName);
+	
+	struct FOldPolygonGroupData
+	{
+		FName Name;
+		TArray<FPolygonID> Polygons;
+	};
+
+	TMap<FPolygonGroupID, FOldPolygonGroupData> OldData;
+	for (const FPolygonGroupID& PolygonGroupID : MeshDescription.PolygonGroups().GetElementIDs())
+	{
+		if (!Remap.Contains(PolygonGroupID) || PolygonGroupID == Remap[PolygonGroupID])
+		{
+			//No need to change this one
+			continue;
+		}
+		FOldPolygonGroupData& PolygonGroupData = OldData.FindOrAdd(PolygonGroupID);
+		PolygonGroupData.Name = PolygonGroupNames[PolygonGroupID];
+		FMeshPolygonGroup& PolygonGroup = MeshDescription.GetPolygonGroup(PolygonGroupID);
+		PolygonGroupData.Polygons = PolygonGroup.Polygons;
+		PolygonGroup.Polygons.Empty();
+		MeshDescription.DeletePolygonGroup(PolygonGroupID);
+	}
+	for (auto Kvp : OldData)
+	{
+		FPolygonGroupID GroupID = Kvp.Key;
+		FPolygonGroupID ToGroupID = Remap[GroupID];
+		if (!MeshDescription.PolygonGroups().IsValid(ToGroupID))
+		{
+			MeshDescription.CreatePolygonGroupWithID(ToGroupID);
+		}
+		TArray<FPolygonID>& Polygons = MeshDescription.GetPolygonGroup(ToGroupID).Polygons;
+		Polygons.Append(Kvp.Value.Polygons);
+		PolygonGroupNames[ToGroupID] = Kvp.Value.Name;
+		for (FPolygonID PolygonID : Polygons)
+		{
+			MeshDescription.GetPolygon(PolygonID).PolygonGroupID = ToGroupID;
+		}
+	}
+}
+
 #undef LOCTEXT_NAMESPACE
diff --git a/Engine/Source/Developer/MeshDescriptionOperations/Public/MeshDescriptionOperations.h b/Engine/Source/Developer/MeshDescriptionOperations/Public/MeshDescriptionOperations.h
--- a/Engine/Source/Developer/MeshDescriptionOperations/Public/MeshDescriptionOperations.h
+++ b/Engine/Source/Developer/MeshDescriptionOperations/Public/MeshDescriptionOperations.h
@@ -8,6 +8,7 @@ struct FMeshDescription;
 struct FRawMesh;
 struct FOverlappingCorners;
 enum class ELightmapUVVersion : int32;
+struct FPolygonGroupID;
 
 DECLARE_LOG_CATEGORY_EXTERN(LogMeshDescriptionOperations, Log, All);
 
@@ -32,6 +33,11 @@ public:
 	/** Convert old FRawMesh format to MeshDescription. */
 	static void ConvertFromRawMesh(const FRawMesh& SourceRawMesh, FMeshDescription& DestinationMeshDescription, const TMap<int32, FName>& MaterialMap);
 
+	/*
+	 * Check if all normals and tangents are valid, if not recompute them
+	 */
+	static void RecomputeNormalsAndTangentsIfNeeded(FMeshDescription& MeshDescription, ETangentOptions TangentOptions, bool bUseMikkTSpace, bool bForceRecomputeNormals = false, bool bForceRecomputeTangents = false);
+
 	/**
 	 * Compute normal, tangent and Bi-Normal for every polygon in the mesh description. (this do not compute Vertex NTBs)
 	 * It also remove the degenerated polygon from the mesh description
@@ -69,4 +75,6 @@ public:
 	static void ConvertHardEdgesToSmoothGroup(const FMeshDescription& SourceMeshDescription, TArray<uint32>& FaceSmoothingMasks);
 
 	static void ConvertSmoothGroupToHardEdges(const TArray<uint32>& FaceSmoothingMasks, FMeshDescription& DestinationMeshDescription);
+
+	static void RemapPolygonGroups(FMeshDescription& MeshDescription, TMap<FPolygonGroupID, FPolygonGroupID>& Remap);
 };
\ No newline at end of file
diff --git a/Engine/Source/Runtime/Engine/Classes/Engine/StaticMesh.h b/Engine/Source/Runtime/Engine/Classes/Engine/StaticMesh.h
--- a/Engine/Source/Runtime/Engine/Classes/Engine/StaticMesh.h
+++ b/Engine/Source/Runtime/Engine/Classes/Engine/StaticMesh.h
@@ -758,7 +758,16 @@ public:
 	ENGINE_API void LoadMeshDescriptions();
 	ENGINE_API void UnloadMeshDescriptions();
 
+	/*
+	 * Return the MeshDescription associate to the LODIndex. The mesh description can be created on the fly if it was null
+	 * and there is a FRawMesh data for this LODIndex.
+	 */
 	ENGINE_API FMeshDescription* GetOriginalMeshDescription(int32 LodIndex);
+	/*
+	 * The const flavor set a new MeshDescription if its null and there is a valid FRawMesh.
+	 * It will simply convert the FRawMesh or copy a valid sourcemodel->MeshDescription pointer into the MeshDescription parameter, .
+	 */
+	ENGINE_API void GetOriginalMeshDescription(int32 LodIndex, FMeshDescription& MeshDescription) const;
 	ENGINE_API FMeshDescription* CreateOriginalMeshDescription(int32 LodIndex);
 	ENGINE_API void CommitOriginalMeshDescription(int32 LodIndex);
 	ENGINE_API void ClearOriginalMeshDescription(int32 LodIndex);
diff --git a/Engine/Source/Runtime/Engine/Private/StaticMesh.cpp b/Engine/Source/Runtime/Engine/Private/StaticMesh.cpp
--- a/Engine/Source/Runtime/Engine/Private/StaticMesh.cpp
+++ b/Engine/Source/Runtime/Engine/Private/StaticMesh.cpp
@@ -1600,29 +1600,13 @@ void FStaticMeshRenderData::Cache(UStaticMesh* Owner, const FStaticMeshLODSettin
 			Args.Add(TEXT("StaticMeshName"), FText::FromString( Owner->GetName() ) );
 			FStaticMeshStatusMessageContext StatusContext( FText::Format( NSLOCTEXT("Engine", "BuildingStaticMeshStatus", "Building static mesh {StaticMeshName}..."), Args ) );
 
-			bool bUseMeshDescription = false;
-			if (Owner->GetOriginalMeshDescription(0) != nullptr)
-			{
-				bUseMeshDescription = true;
-			}
+			check(Owner->GetOriginalMeshDescription(0) != nullptr)
 
-			if (bUseMeshDescription)
-			{
-				IMeshBuilderModule& MeshBuilderModule = FModuleManager::Get().LoadModuleChecked<IMeshBuilderModule>(TEXT("MeshBuilder"));
-				if (!MeshBuilderModule.BuildMesh(*this, Owner, LODGroup))
-				{
-					UE_LOG(LogStaticMesh, Error, TEXT("Failed to build static mesh. See previous line(s) for details."));
-					return;
-				}
-			}
-			else
+			IMeshBuilderModule& MeshBuilderModule = FModuleManager::Get().LoadModuleChecked<IMeshBuilderModule>(TEXT("MeshBuilder"));
+			if (!MeshBuilderModule.BuildMesh(*this, Owner, LODGroup))
 			{
-				IMeshUtilities& MeshUtilities = FModuleManager::Get().LoadModuleChecked<IMeshUtilities>(TEXT("MeshUtilities"));
-				if (!MeshUtilities.BuildStaticMesh(*this, Owner, LODGroup))
-				{
-					UE_LOG(LogStaticMesh, Error, TEXT("Failed to build static mesh. See previous line(s) for details."));
-					return;
-				}
+				UE_LOG(LogStaticMesh, Error, TEXT("Failed to build static mesh. See previous line(s) for details."));
+				return;
 			}
 
 			ComputeUVDensities();
@@ -2817,6 +2801,28 @@ void UStaticMesh::UnloadMeshDescriptions()
 	MeshDescriptions = nullptr;
 }
 
+void UStaticMesh::GetOriginalMeshDescription(int32 LodIndex, FMeshDescription& MeshDescription) const
+{
+	if (SourceModels.IsValidIndex(LodIndex))
+	{
+		check(MeshDescriptions->Num() == SourceModels.Num());
+		check(MeshDescriptions->Get(LodIndex) == SourceModels[LodIndex].OriginalMeshDescription);
+
+		if (MeshDescriptions->Get(LodIndex) == nullptr && !SourceModels[LodIndex].IsRawMeshEmpty())
+		{
+			FRawMesh LodRawMesh;
+			SourceModels[LodIndex].LoadRawMesh(LodRawMesh);
+			TMap<int32, FName> MaterialMap;
+			FillMaterialName(StaticMaterials, MaterialMap);
+			FMeshDescriptionOperations::ConvertFromRawMesh(LodRawMesh, MeshDescription, MaterialMap);
+		}
+		else if(MeshDescriptions->Get(LodIndex) != nullptr)
+		{
+			MeshDescription = *MeshDescriptions->Get(LodIndex);
+		}
+	}
+}
+
 FMeshDescription* UStaticMesh::GetOriginalMeshDescription(int32 LodIndex)
 {
 	LoadMeshDescriptions();
