Fix main merge conflict

diff --git a/Engine/Plugins/Experimental/ProxyLODPlugin/Source/ProxyLOD/Private/ProxyLODVolume.cpp b/Engine/Plugins/Experimental/ProxyLODPlugin/Source/ProxyLOD/Private/ProxyLODVolume.cpp
--- a/Engine/Plugins/Experimental/ProxyLODPlugin/Source/ProxyLOD/Private/ProxyLODVolume.cpp
+++ b/Engine/Plugins/Experimental/ProxyLODPlugin/Source/ProxyLOD/Private/ProxyLODVolume.cpp
@@ -13,6 +13,8 @@
 #include <openvdb/tools/MeshToVolume.h> // for MeshToVolume
 #include <openvdb/tools/VolumeToMesh.h> // for VolumeToMesh
 
+#include "MeshDescription.h"
+
 typedef openvdb::math::Transform	OpenVDBTransform;
 
 class FProxyLODVolumeImpl : public IProxyLODVolume
@@ -32,7 +34,7 @@ public:
 
 	bool Initialize(const TArray<FMeshMergeData>& Geometry, float Accuracy)
 	{
-		FRawMeshArrayAdapter SrcGeometryAdapter(Geometry);
+		FMeshDescriptionArrayAdapter SrcGeometryAdapter(Geometry);
 		OpenVDBTransform::Ptr XForm = OpenVDBTransform::createLinearTransform(Accuracy);
 		SrcGeometryAdapter.SetTransform(XForm);
 
@@ -78,25 +80,22 @@ public:
 		return Sampler->wsSample(openvdb::Vec3R(Point.X, Point.Y, Point.Z));
 	}
 
-	virtual void ConvertToRawMesh(FRawMesh& OutRawMesh) const override
+	virtual void ConvertToRawMesh(FMeshDescription& OutRawMesh) const override
 	{
 		// Mesh types that will be shared by various stages.
 		FAOSMesh AOSMeshedVolume;
 		ProxyLOD::SDFVolumeToMesh(SDFVolume, 0.0, 0.0, AOSMeshedVolume);
-
-		FVertexDataMesh VertexDataMesh;
-		ProxyLOD::ConvertMesh(AOSMeshedVolume, VertexDataMesh);
-
-		ProxyLOD::ConvertMesh(VertexDataMesh, OutRawMesh);
+		ProxyLOD::ConvertMesh(AOSMeshedVolume, OutRawMesh);
 	}
 
-	void ExpandNarrowBand(float ExteriorWidth, float InteriorWidth)
+	void ExpandNarrowBand(float ExteriorWidth, float InteriorWidth) override
 	{
 		using namespace openvdb::tools;
 
-		FRawMesh RawMesh;
+		FMeshDescription RawMesh;
+		UStaticMesh::RegisterMeshAttributes(RawMesh);
 		ConvertToRawMesh(RawMesh);
-		FRawMeshAdapter MeshAdapter(RawMesh, SDFVolume->transform());
+		FMeshDescriptionAdapter MeshAdapter(RawMesh, SDFVolume->transform());
 
 		openvdb::FloatGrid::Ptr NewSDFVolume;
 		openvdb::Int32Grid::Ptr NewSrcPolyIndexGrid;
diff --git a/Engine/Plugins/Experimental/ProxyLODPlugin/Source/ProxyLOD/Public/ProxyLODVolume.h b/Engine/Plugins/Experimental/ProxyLODPlugin/Source/ProxyLOD/Public/ProxyLODVolume.h
--- a/Engine/Plugins/Experimental/ProxyLODPlugin/Source/ProxyLOD/Public/ProxyLODVolume.h
+++ b/Engine/Plugins/Experimental/ProxyLODPlugin/Source/ProxyLOD/Public/ProxyLODVolume.h
@@ -5,7 +5,7 @@
 #include "CoreMinimal.h"
 
 struct FMeshMergeData;
-struct FRawMesh;
+struct FMeshDescription;
 
 class PROXYLODMESHREDUCTION_API IProxyLODVolume
 {
@@ -45,7 +45,7 @@ public:
 	virtual void CloseGaps(const double GapRadius, const int32 MaxDilations) = 0;
 
 	/** Extract iso distance 0 from OpenVDB volume as a RawMesh */
-	virtual void ConvertToRawMesh(FRawMesh& OutRawMesh) const = 0;
+	virtual void ConvertToRawMesh(FMeshDescription& OutRawMesh) const = 0;
 
 	/** Expand exterior and interior narrow band of OpenVDB volume by given amount */
 	virtual void ExpandNarrowBand(float ExteriorWidth, float InteriorWidth) = 0;
