Emissive isn't baked correctly in TM-MeshbakeMap
#fix Amalgamate all flat material data when creating a proxy

diff --git a/Engine/Plugins/Experimental/ProxyLODPlugin/Source/ProxyLOD/Private/ProxyLODPlugin.cpp b/Engine/Plugins/Experimental/ProxyLODPlugin/Source/ProxyLOD/Private/ProxyLODPlugin.cpp
--- a/Engine/Plugins/Experimental/ProxyLODPlugin/Source/ProxyLOD/Private/ProxyLODPlugin.cpp
+++ b/Engine/Plugins/Experimental/ProxyLODPlugin/Source/ProxyLOD/Private/ProxyLODPlugin.cpp
@@ -929,9 +929,13 @@ void FVoxelizeMeshMerging::ProxyLOD(const FMeshMergeDataArray& InData, const FMe
 				{
 					const auto& FirstMaterial = (*BakedMaterials)[0];
 					OutMaterial.BlendMode              = FirstMaterial.BlendMode;
-					OutMaterial.bTwoSided              = FirstMaterial.bTwoSided;
-					OutMaterial.bDitheredLODTransition = FirstMaterial.bDitheredLODTransition;
-					OutMaterial.EmissiveScale          = FirstMaterial.EmissiveScale;
+
+					for (const FFlattenMaterial& FlatMaterial : (*BakedMaterials))
+					{
+						OutMaterial.bTwoSided |= FlatMaterial.bTwoSided;
+						OutMaterial.bDitheredLODTransition |= FlatMaterial.bDitheredLODTransition;
+						OutMaterial.EmissiveScale = FMath::Max(OutMaterial.EmissiveScale, FlatMaterial.EmissiveScale);
+					}
 				}
 
 				// --- Map a correspondence between texels in the texture atlas and locations on the source geometry --
