Fixed another deprecation warning.

diff --git a/Engine/Plugins/Experimental/ProxyLODPlugin/Source/ProxyLOD/Private/ProxyLODMeshConvertUtils.cpp b/Engine/Plugins/Experimental/ProxyLODPlugin/Source/ProxyLOD/Private/ProxyLODMeshConvertUtils.cpp
--- a/Engine/Plugins/Experimental/ProxyLODPlugin/Source/ProxyLOD/Private/ProxyLODMeshConvertUtils.cpp
+++ b/Engine/Plugins/Experimental/ProxyLODPlugin/Source/ProxyLOD/Private/ProxyLODMeshConvertUtils.cpp
@@ -61,7 +61,8 @@ void ProxyLOD::MixedPolyMeshToRawMesh(const FMixedPolyMesh& SimpleMesh, FMeshDes
 	// Connectivity:
 	auto CreateTriangle = [&DstRawMesh, PolygonGroupID, &VertexInstanceNormals, &VertexInstanceTangents, &VertexInstanceBinormalSigns, &VertexInstanceColors, &VertexInstanceUVs, &EdgeHardnesses, &EdgeCreaseSharpnesses](FVertexID TriangleIndex[3])
 	{
-		FVertexInstanceID VertexInstanceIDs[3];
+		TArray<FVertexInstanceID> VertexInstanceIDs;
+		VertexInstanceIDs.SetNum(3);
 		for (int32 Corner = 0; Corner < 3; ++Corner)
 		{
 			VertexInstanceIDs[Corner] = DstRawMesh.CreateVertexInstance(TriangleIndex[Corner]);
@@ -74,33 +75,8 @@ void ProxyLOD::MixedPolyMeshToRawMesh(const FMixedPolyMesh& SimpleMesh, FMeshDes
 			VertexInstanceUVs.Set(VertexInstanceIDs[Corner], 0, FVector2D(0.0f, 0.0f));
 		}
 
-		//Create a polygon from this triangle
-		TArray<FMeshDescription::FContourPoint> Contours;
-		for (int32 Corner = 0; Corner < 3; ++Corner)
-		{
-			int32 ContourPointIndex = Contours.AddDefaulted();
-			FMeshDescription::FContourPoint& ContourPoint = Contours[ContourPointIndex];
-			//Find the matching edge ID
-			uint32 CornerIndices[2];
-			CornerIndices[0] = (Corner + 0) % 3;
-			CornerIndices[1] = (Corner + 1) % 3;
-
-			FVertexID EdgeVertexIDs[2];
-			EdgeVertexIDs[0] = TriangleIndex[CornerIndices[0]];
-			EdgeVertexIDs[1] = TriangleIndex[CornerIndices[1]];
-
-			FEdgeID MatchEdgeId = DstRawMesh.GetVertexPairEdge(EdgeVertexIDs[0], EdgeVertexIDs[1]);
-			if (MatchEdgeId == FEdgeID::Invalid)
-			{
-				MatchEdgeId = DstRawMesh.CreateEdge(EdgeVertexIDs[0], EdgeVertexIDs[1]);
-				EdgeHardnesses[MatchEdgeId] = false;
-				EdgeCreaseSharpnesses[MatchEdgeId] = 0.0f;
-			}
-			ContourPoint.EdgeID = MatchEdgeId;
-			ContourPoint.VertexInstanceID = VertexInstanceIDs[CornerIndices[0]];
-		}
 		// Insert a polygon into the mesh
-		const FPolygonID NewPolygonID = DstRawMesh.CreatePolygon(PolygonGroupID, Contours);
+		const FPolygonID NewPolygonID = DstRawMesh.CreatePolygon(PolygonGroupID, VertexInstanceIDs);
 		//Triangulate the polygon
 		FMeshPolygon& Polygon = DstRawMesh.GetPolygon(NewPolygonID);
 		DstRawMesh.ComputePolygonTriangulation(NewPolygonID, Polygon.Triangles);
