Fixed deprecation warnings.

diff --git a/Engine/Plugins/Experimental/ProxyLODPlugin/Source/ProxyLOD/Private/ProxyLODMeshConvertUtils.cpp b/Engine/Plugins/Experimental/ProxyLODPlugin/Source/ProxyLOD/Private/ProxyLODMeshConvertUtils.cpp
--- a/Engine/Plugins/Experimental/ProxyLODPlugin/Source/ProxyLOD/Private/ProxyLODMeshConvertUtils.cpp
+++ b/Engine/Plugins/Experimental/ProxyLODPlugin/Source/ProxyLOD/Private/ProxyLODMeshConvertUtils.cpp
@@ -174,7 +174,8 @@ void ProxyLOD::AOSMeshToRawMesh(const FAOSMesh& AOSMesh, FMeshDescription& OutRa
 	// Connectivity: 
 	auto CreateTriangle = [&OutRawMesh, PolygonGroupID, &VertexInstanceNormals, &VertexInstanceTangents, &VertexInstanceBinormalSigns, &VertexInstanceColors, &VertexInstanceUVs, &EdgeHardnesses, &EdgeCreaseSharpnesses](const FVertexID TriangleIndex[3], const FVector Normals[3])
 	{
-		FVertexInstanceID VertexInstanceIDs[3];
+		TArray<FVertexInstanceID> VertexInstanceIDs;
+		VertexInstanceIDs.SetNum(3);
 		for (int32 Corner = 0; Corner < 3; ++Corner)
 		{
 			VertexInstanceIDs[Corner] = OutRawMesh.CreateVertexInstance(TriangleIndex[Corner]);
@@ -187,33 +188,8 @@ void ProxyLOD::AOSMeshToRawMesh(const FAOSMesh& AOSMesh, FMeshDescription& OutRa
 			VertexInstanceUVs.Set(VertexInstanceIDs[Corner], 0, FVector2D(0.0f, 0.0f));
 		}
 
-		//Create a polygon from this triangle
-		TArray<FMeshDescription::FContourPoint> Contours;
-		for (int32 Corner = 0; Corner < 3; ++Corner)
-		{
-			int32 ContourPointIndex = Contours.AddDefaulted();
-			FMeshDescription::FContourPoint& ContourPoint = Contours[ContourPointIndex];
-			//Find the matching edge ID
-			uint32 CornerIndices[2];
-			CornerIndices[0] = (Corner + 0) % 3;
-			CornerIndices[1] = (Corner + 1) % 3;
-
-			FVertexID EdgeVertexIDs[2];
-			EdgeVertexIDs[0] = TriangleIndex[CornerIndices[0]];
-			EdgeVertexIDs[1] = TriangleIndex[CornerIndices[1]];
-
-			FEdgeID MatchEdgeId = OutRawMesh.GetVertexPairEdge(EdgeVertexIDs[0], EdgeVertexIDs[1]);
-			if (MatchEdgeId == FEdgeID::Invalid)
-			{
-				MatchEdgeId = OutRawMesh.CreateEdge(EdgeVertexIDs[0], EdgeVertexIDs[1]);
-				EdgeHardnesses[MatchEdgeId] = false;
-				EdgeCreaseSharpnesses[MatchEdgeId] = 0.0f;
-			}
-			ContourPoint.EdgeID = MatchEdgeId;
-			ContourPoint.VertexInstanceID = VertexInstanceIDs[CornerIndices[0]];
-		}
 		// Insert a polygon into the mesh
-		const FPolygonID NewPolygonID = OutRawMesh.CreatePolygon(PolygonGroupID, Contours);
+		const FPolygonID NewPolygonID = OutRawMesh.CreatePolygon(PolygonGroupID, VertexInstanceIDs);
 		//Triangulate the polygon
 		FMeshPolygon& Polygon = OutRawMesh.GetPolygon(NewPolygonID);
 		OutRawMesh.ComputePolygonTriangulation(NewPolygonID, Polygon.Triangles);
@@ -296,7 +272,8 @@ void ProxyLOD::VertexDataMeshToRawMesh(const FVertexDataMesh& SrcVertexDataMesh,
 		int32 TriangleIndex = TriangleIndices[0] / 3;
 		
 		FVertexID VertexIndexes[3];
-		FVertexInstanceID VertexInstanceIDs[3];
+		TArray<FVertexInstanceID> VertexInstanceIDs;
+		VertexInstanceIDs.SetNum(3);
 		for (int32 Corner = 0; Corner < 3; ++Corner)
 		{
 			int32 SrcIndex = SrcVertexDataMesh.Indices[TriangleIndices[Corner]];
@@ -348,33 +325,8 @@ void ProxyLOD::VertexDataMeshToRawMesh(const FVertexDataMesh& SrcVertexDataMesh,
 			}
 		}
 
-		//Create a polygon from this triangle
-		TArray<FMeshDescription::FContourPoint> Contours;
-		for (int32 Corner = 0; Corner < 3; ++Corner)
-		{
-			int32 ContourPointIndex = Contours.AddDefaulted();
-			FMeshDescription::FContourPoint& ContourPoint = Contours[ContourPointIndex];
-			//Find the matching edge ID
-			uint32 CornerIndices[2];
-			CornerIndices[0] = (Corner + 0) % 3;
-			CornerIndices[1] = (Corner + 1) % 3;
-
-			FVertexID EdgeVertexIDs[2];
-			EdgeVertexIDs[0] = VertexIndexes[CornerIndices[0]];
-			EdgeVertexIDs[1] = VertexIndexes[CornerIndices[1]];
-
-			FEdgeID MatchEdgeId = OutRawMesh.GetVertexPairEdge(EdgeVertexIDs[0], EdgeVertexIDs[1]);
-			if (MatchEdgeId == FEdgeID::Invalid)
-			{
-				MatchEdgeId = OutRawMesh.CreateEdge(EdgeVertexIDs[0], EdgeVertexIDs[1]);
-				EdgeHardnesses[MatchEdgeId] = false;
-				EdgeCreaseSharpnesses[MatchEdgeId] = 0.0f;
-			}
-			ContourPoint.EdgeID = MatchEdgeId;
-			ContourPoint.VertexInstanceID = VertexInstanceIDs[CornerIndices[0]];
-		}
 		// Insert a polygon into the mesh
-		const FPolygonID NewPolygonID = OutRawMesh.CreatePolygon(PolygonGroupID, Contours);
+		const FPolygonID NewPolygonID = OutRawMesh.CreatePolygon(PolygonGroupID, VertexInstanceIDs);
 		//Triangulate the polygon
 		FMeshPolygon& Polygon = OutRawMesh.GetPolygon(NewPolygonID);
 		OutRawMesh.ComputePolygonTriangulation(NewPolygonID, Polygon.Triangles);
