ProxyLOD update copy-up from ModelingToolsFramework:
6620219 Don't use FMeshDescription VertexInstanceArray as an index buffer.  Fixing certain assumptions that were made in the conversion from FRawMesh.
6620234 Expose voxel-based CSG operations so other tools can uses these.
#rnx

diff --git a/Engine/Plugins/Experimental/ProxyLODPlugin/Source/ProxyLOD/Private/ProxyLODMeshConvertUtils.h b/Engine/Plugins/Experimental/ProxyLODPlugin/Source/ProxyLOD/Private/ProxyLODMeshConvertUtils.h
--- a/Engine/Plugins/Experimental/ProxyLODPlugin/Source/ProxyLOD/Private/ProxyLODMeshConvertUtils.h
+++ b/Engine/Plugins/Experimental/ProxyLODPlugin/Source/ProxyLOD/Private/ProxyLODMeshConvertUtils.h
@@ -39,9 +39,10 @@ namespace ProxyLOD
 	*
 	* @param InMesh   Mesh with a mixture of quads and triangles.
 	* @param OutMesh  Triangle Mesh.
+	* @param bClockWise how to order the verts in a triangle
 	*/
 	template <typename T>
-	void MixedPolyMeshToAOSMesh(const FMixedPolyMesh& InMesh, TAOSMesh<T>& OutMesh);
+	void MixedPolyMeshToAOSMesh(const FMixedPolyMesh& InMesh, TAOSMesh<T>& OutMesh, const bool bClockWise = true);
 
 
 	/**
@@ -103,8 +104,6 @@ namespace ProxyLOD
 #define PROXYLOD_CLOCKWISE_TRIANGLES  1
 #endif
 
-
-
 static FVector ComputeNormal(const FVector(&Tri)[3])
 {
 
@@ -119,7 +118,7 @@ static FVector ComputeNormal(const FVector(&Tri)[3])
 
 // Convert MixedPolyMesh to AOS Mesh.  This requires splitting quads to produce triangles.
 template <typename T>
-void ProxyLOD::MixedPolyMeshToAOSMesh(const FMixedPolyMesh& MixedPolyMesh, TAOSMesh<T>& DstAOSMesh)
+void ProxyLOD::MixedPolyMeshToAOSMesh(const FMixedPolyMesh& MixedPolyMesh, TAOSMesh<T>& DstAOSMesh, bool bClockWise)
 {
 
 	// Splitting a quad doesn't introduce any new verts.
@@ -160,7 +159,7 @@ void ProxyLOD::MixedPolyMeshToAOSMesh(const FMixedPolyMesh& MixedPolyMesh, TAOSM
 		// NB: The Quads are ordered in clockwise fashion.
 
 		ProxyLOD::Parallel_For(ProxyLOD::FUIntRange(0, NumQuads),
-			[&MixedPolyMesh, &DstAOSMesh](const ProxyLOD::FUIntRange& Range)
+			[&MixedPolyMesh, &DstAOSMesh, bClockWise](const ProxyLOD::FUIntRange& Range)
 		{
 			uint32* Indices = DstAOSMesh.Indexes;
 			for (uint32 q = Range.begin(), Q = Range.end(); q < Q; ++q)
@@ -168,33 +167,36 @@ void ProxyLOD::MixedPolyMeshToAOSMesh(const FMixedPolyMesh& MixedPolyMesh, TAOSM
 				const uint32 Offset = q * 6;
 				const openvdb::Vec4I& Quad = MixedPolyMesh.Quads[q];
 				// add as two triangles
-#if   (PROXYLOD_CLOCKWISE_TRIANGLES  == 1)
-				// first triangle
-				Indices[Offset] = Quad[0];
-				Indices[Offset + 1] = Quad[1];
-				Indices[Offset + 2] = Quad[2];
-				// second triangle
-				Indices[Offset + 3] = Quad[2];
-				Indices[Offset + 4] = Quad[3];
-				Indices[Offset + 5] = Quad[0];
-#else
-				// first triangle
-				Indices[Offset] = Quad[0];
-				Indices[Offset + 1] = Quad[3];
-				Indices[Offset + 2] = Quad[2];
-
-				// second triangle
-				Indices[Offset + 3] = Quad[2];
-				Indices[Offset + 4] = Quad[1];
-				Indices[Offset + 5] = Quad[0];
-#endif
+				if (bClockWise)
+				{
+					// first triangle
+					Indices[Offset] = Quad[0];
+					Indices[Offset + 1] = Quad[1];
+					Indices[Offset + 2] = Quad[2];
+					// second triangle
+					Indices[Offset + 3] = Quad[2];
+					Indices[Offset + 4] = Quad[3];
+					Indices[Offset + 5] = Quad[0];
+				}
+				else
+				{
+					// first triangle
+					Indices[Offset] = Quad[0];
+					Indices[Offset + 1] = Quad[3];
+					Indices[Offset + 2] = Quad[2];
+
+					// second triangle
+					Indices[Offset + 3] = Quad[2];
+					Indices[Offset + 4] = Quad[1];
+					Indices[Offset + 5] = Quad[0];
+				}
 			}
 
 		});
 
 		// add the MixedPolyMesh triangles.
 		ProxyLOD::Parallel_For(ProxyLOD::FUIntRange(0, MixedPolyMesh.Triangles.size()),
-			[&MixedPolyMesh, &DstAOSMesh, NumQuads](const ProxyLOD::FUIntRange& Range)
+			[&MixedPolyMesh, &DstAOSMesh, NumQuads, bClockWise](const ProxyLOD::FUIntRange& Range)
 		{
 			uint32* Indices = DstAOSMesh.Indexes;
 			for (uint32 t = Range.begin(), EndT = Range.end(); t < EndT; ++t)
@@ -202,15 +204,18 @@ void ProxyLOD::MixedPolyMeshToAOSMesh(const FMixedPolyMesh& MixedPolyMesh, TAOSM
 				const uint32 Offset = NumQuads * 6 + t * 3;
 				const openvdb::Vec3I& Tri = MixedPolyMesh.Triangles[t];
 				// add the triangle
-#if (PROXYLOD_CLOCKWISE_TRIANGLES == 1)
-				Indices[Offset] = Tri[0];
-				Indices[Offset + 1] = Tri[1];
-				Indices[Offset + 2] = Tri[2];
-#else
-				Indices[Offset] = Tri[2];
-				Indices[Offset + 1] = Tri[1];
-				Indices[Offset + 2] = Tri[0];
-#endif 
+				if (bClockWise)
+				{
+					Indices[Offset] = Tri[0];
+					Indices[Offset + 1] = Tri[1];
+					Indices[Offset + 2] = Tri[2];
+				}
+				else
+				{
+					Indices[Offset] = Tri[2];
+					Indices[Offset + 1] = Tri[1];
+					Indices[Offset + 2] = Tri[0];
+				}
 			}
 		});
 	}
diff --git a/Engine/Plugins/Experimental/ProxyLODPlugin/Source/ProxyLOD/Private/ProxyLODMeshTypes.cpp b/Engine/Plugins/Experimental/ProxyLODPlugin/Source/ProxyLOD/Private/ProxyLODMeshTypes.cpp
--- a/Engine/Plugins/Experimental/ProxyLODPlugin/Source/ProxyLOD/Private/ProxyLODMeshTypes.cpp
+++ b/Engine/Plugins/Experimental/ProxyLODPlugin/Source/ProxyLOD/Private/ProxyLODMeshTypes.cpp
@@ -6,7 +6,7 @@
 // --- FMeshDescriptionAdapter ----
 
 FMeshDescriptionAdapter::FMeshDescriptionAdapter(const FMeshDescription& InRawMesh, const openvdb::math::Transform& InTransform) :
-	RawMesh(&InRawMesh), Transform(&InTransform)
+	RawMesh(&InRawMesh), Transform(InTransform)
 {
 	InitializeCacheData();
 }
@@ -25,6 +25,21 @@ void FMeshDescriptionAdapter::InitializeCacheData()
 	{
 		TriangleCount += RawMesh->GetPolygonTriangles(PolygonID).Num();
 	}
+
+	IndexBuffer.Reserve(TriangleCount * 3);
+
+	for (const FPolygonID& PolygonID : RawMesh->Polygons().GetElementIDs())
+	{
+		const auto& Triangles = RawMesh->GetPolygonTriangles(PolygonID);
+		for (const FMeshTriangle& Tri : Triangles)
+		{
+			IndexBuffer.Add(Tri.VertexInstanceID0);
+			IndexBuffer.Add(Tri.VertexInstanceID1);
+			IndexBuffer.Add(Tri.VertexInstanceID2);
+		}
+	}
+
+
 }
 
 size_t FMeshDescriptionAdapter::polygonCount() const
@@ -40,10 +55,13 @@ size_t FMeshDescriptionAdapter::pointCount() const
 void FMeshDescriptionAdapter::getIndexSpacePoint(size_t FaceNumber, size_t CornerNumber, openvdb::Vec3d& pos) const
 {
 	// Get the vertex position in local space.
-	const FVertexInstanceID VertexInstanceID = FVertexInstanceID(FaceNumber * 3 + CornerNumber);
+	const FVertexInstanceID VertexInstanceID = IndexBuffer[FaceNumber * 3 + CornerNumber];
 	// float3 position 
-	FVector Position = VertexPositions[RawMesh->GetVertexInstanceVertex(VertexInstanceID)];
-	pos = Transform->worldToIndex(openvdb::Vec3d(Position.X, Position.Y, Position.Z));
+	const FVertexID VertexID = RawMesh->GetVertexInstanceVertex(VertexInstanceID);
+
+
+	FVector Position = VertexPositions[VertexID];
+	pos = Transform.worldToIndex(openvdb::Vec3d(Position.X, Position.Y, Position.Z));
 };
 
 
@@ -63,11 +81,30 @@ FMeshDescriptionArrayAdapter::FMeshDescriptionArrayAdapter(const TArray<const FM
 		FMeshDescription *RawMesh = MergeData->RawMesh;
 
 		PointCount += size_t(RawMesh->Vertices().Num());
+		// Sum up all the polys in this mesh.
+		int32 MeshPolyCount = 0;
 		for (const FPolygonID& PolygonID : RawMesh->Polygons().GetElementIDs())
 		{
-			PolyCount += RawMesh->GetPolygonTriangles(PolygonID).Num();
+			MeshPolyCount += RawMesh->GetPolygonTriangles(PolygonID).Num();
 		}
 
+		// Construct local index buffer for this mesh
+		IndexBufferArray.push_back(std::vector<FVertexInstanceID>());
+		std::vector<FVertexInstanceID>& IndexBuffer = IndexBufferArray[MeshIdx];
+		IndexBuffer.reserve(MeshPolyCount * 3);
+		for (const FPolygonID& PolygonID : RawMesh->Polygons().GetElementIDs())
+		{
+			const auto& Triangles = RawMesh->GetPolygonTriangles(PolygonID);
+			for (const FMeshTriangle& Tri : Triangles)
+			{
+				IndexBuffer.push_back(Tri.VertexInstanceID0);
+				IndexBuffer.push_back(Tri.VertexInstanceID1);
+				IndexBuffer.push_back(Tri.VertexInstanceID2);
+			}
+		}
+
+
+		PolyCount += MeshPolyCount;
 		PolyOffsetArray.push_back(PolyCount);
 		RawMeshArray.push_back(RawMesh);
 		RawMeshArrayData.push_back(FMeshDescriptionAttributesGetter(RawMesh));
@@ -93,10 +130,29 @@ FMeshDescriptionArrayAdapter::FMeshDescriptionArrayAdapter(const TArray<FMeshMer
 		
 		FMeshDescription *RawMesh = MergeData->RawMesh;
 		PointCount += size_t(RawMesh->Vertices().Num());
+
+		// Sum up all the polys in this mesh.
+		int32 MeshPolyCount = 0;
 		for (const FPolygonID& PolygonID : RawMesh->Polygons().GetElementIDs())
 		{
-			PolyCount += RawMesh->GetPolygonTriangles(PolygonID).Num();
+			MeshPolyCount += RawMesh->GetPolygonTriangles(PolygonID).Num();
 		}
+		PolyCount += MeshPolyCount;
+		// Construct local index buffer for this mesh
+		IndexBufferArray.push_back(std::vector<FVertexInstanceID>());
+		std::vector<FVertexInstanceID>& IndexBuffer = IndexBufferArray[MeshIdx];
+		IndexBuffer.reserve(MeshPolyCount * 3);
+		for (const FPolygonID& PolygonID : RawMesh->Polygons().GetElementIDs())
+		{
+			const auto& Triangles = RawMesh->GetPolygonTriangles(PolygonID);
+			for (const FMeshTriangle& Tri : Triangles)
+			{
+				IndexBuffer.push_back(Tri.VertexInstanceID0);
+				IndexBuffer.push_back(Tri.VertexInstanceID1);
+				IndexBuffer.push_back(Tri.VertexInstanceID2);
+			}
+		}
+
 
 		PolyOffsetArray.push_back(PolyCount);
 		RawMeshArray.push_back(RawMesh);
@@ -121,11 +177,30 @@ FMeshDescriptionArrayAdapter::FMeshDescriptionArrayAdapter(const TArray<FMeshMer
 		
 		FMeshDescription *RawMesh = MergeData->RawMesh;
 		PointCount += size_t(RawMesh->Vertices().Num());
+		
+		// Sum up all the polys in this mesh.
+		int32 MeshPolyCount = 0;
+		for (const FPolygonID& PolygonID : RawMesh->Polygons().GetElementIDs())
+		{
+			MeshPolyCount += RawMesh->GetPolygonTriangles(PolygonID).Num();
+		}
+
+		// Construct local index buffer for this mesh
+		IndexBufferArray.push_back(std::vector<FVertexInstanceID>());
+		std::vector<FVertexInstanceID>& IndexBuffer = IndexBufferArray[MeshIdx];
+		IndexBuffer.reserve(MeshPolyCount * 3);
 		for (const FPolygonID& PolygonID : RawMesh->Polygons().GetElementIDs())
 		{
-			PolyCount += RawMesh->GetPolygonTriangles(PolygonID).Num();
+			const auto& Triangles = RawMesh->GetPolygonTriangles(PolygonID);
+			for (const FMeshTriangle& Tri : Triangles)
+			{
+				IndexBuffer.push_back(Tri.VertexInstanceID0);
+				IndexBuffer.push_back(Tri.VertexInstanceID1);
+				IndexBuffer.push_back(Tri.VertexInstanceID2);
+			}
 		}
 
+		PolyCount += MeshPolyCount;
 		PolyOffsetArray.push_back(PolyCount);
 		RawMeshArray.push_back(RawMesh);
 		RawMeshArrayData.push_back(FMeshDescriptionAttributesGetter(RawMesh));
@@ -147,6 +222,13 @@ FMeshDescriptionArrayAdapter::FMeshDescriptionArrayAdapter(const FMeshDescriptio
 	{
 		RawMeshArrayData.push_back(FMeshDescriptionAttributesGetter(RawMesh));
 	}
+
+	IndexBufferArray.reserve(other.IndexBufferArray.size());
+	for (const auto& IndexBuffer : other.IndexBufferArray)
+	{
+		IndexBufferArray.push_back(IndexBuffer);
+	}
+
 }
 
 FMeshDescriptionArrayAdapter::~FMeshDescriptionArrayAdapter()
@@ -166,8 +248,9 @@ void FMeshDescriptionArrayAdapter::getWorldSpacePoint(size_t FaceNumber, size_t
 	const FMeshDescription& RawMesh = GetRawMesh(FaceNumber, MeshIdx, LocalFaceNumber, &AttributesGetter);
 	check(AttributesGetter);
 
+	const auto& IndexBuffer = IndexBufferArray[MeshIdx];
 	// Get the vertex position in local space.
-	const FVertexInstanceID VertexInstanceID = FVertexInstanceID(3 * LocalFaceNumber + int32(CornerNumber));
+	const FVertexInstanceID VertexInstanceID = IndexBuffer[3 * LocalFaceNumber + int32(CornerNumber)];
 	// float3 position 
 	FVector Position = AttributesGetter->VertexPositions[RawMesh.GetVertexInstanceVertex(VertexInstanceID)];
 	pos = openvdb::Vec3d(Position.X, Position.Y, Position.Z);
diff --git a/Engine/Plugins/Experimental/ProxyLODPlugin/Source/ProxyLOD/Private/ProxyLODMeshTypes.h b/Engine/Plugins/Experimental/ProxyLODPlugin/Source/ProxyLOD/Private/ProxyLODMeshTypes.h
--- a/Engine/Plugins/Experimental/ProxyLODPlugin/Source/ProxyLOD/Private/ProxyLODMeshTypes.h
+++ b/Engine/Plugins/Experimental/ProxyLODPlugin/Source/ProxyLOD/Private/ProxyLODMeshTypes.h
@@ -219,7 +219,7 @@ public:
 	*/
 	const openvdb::math::Transform& GetTransform() const
 	{
-		return *Transform;
+		return Transform;
 	}
 
 private:
@@ -227,14 +227,20 @@ private:
 	// Pointer to the raw mesh we are wrapping
 	const FMeshDescription* RawMesh;
 
+	
+
 	//////////////////////////////////////////////////////////////////////////
 	//Cache data
 	void InitializeCacheData();
 	uint32 TriangleCount;
 	TVertexAttributesConstRef<FVector> VertexPositions;
+	// Local version of the index array.  The FMeshDescription doesn't really have one.
+	TArray<FVertexInstanceID> IndexBuffer;
+
+
 
 	// Transform used to convert the mesh space into the index space used by voxelization
-	const openvdb::math::Transform* Transform;
+	const openvdb::math::Transform Transform;
 };
 
 namespace ProxyLOD
@@ -421,6 +427,9 @@ protected:
 
 	std::vector<const FMeshMergeData*>  MergeDataArray;
 
+	// Need to build local index buffers for each mesh because the FMeshDescription doesn't natively have the construct.
+	std::vector<std::vector<FVertexInstanceID>> IndexBufferArray;
+
 };
 
 /**
diff --git a/Engine/Plugins/Experimental/ProxyLODPlugin/Source/ProxyLOD/Private/ProxyLODVolume.cpp b/Engine/Plugins/Experimental/ProxyLODPlugin/Source/ProxyLOD/Private/ProxyLODVolume.cpp
--- a/Engine/Plugins/Experimental/ProxyLODPlugin/Source/ProxyLOD/Private/ProxyLODVolume.cpp
+++ b/Engine/Plugins/Experimental/ProxyLODPlugin/Source/ProxyLOD/Private/ProxyLODVolume.cpp
@@ -15,6 +15,7 @@ THIRD_PARTY_INCLUDES_START
 #include <openvdb/tools/Interpolation.h> // for Spatial Query
 #include <openvdb/tools/MeshToVolume.h> // for MeshToVolume
 #include <openvdb/tools/VolumeToMesh.h> // for VolumeToMesh
+#include <openvdb/tools/Composite.h> // for CSG operations
 #pragma warning(pop)
 THIRD_PARTY_INCLUDES_END
 
@@ -149,3 +150,345 @@ TUniquePtr<IProxyLODVolume> IProxyLODVolume::CreateSDFVolumeFromMeshArray(const
 
 	return Volume;
 }
+
+
+class FPolygonSoup
+{
+public:
+	FPolygonSoup(const TArray<FMeshDescriptionAdapter>& AdapterArray, double VoxelSize);
+
+	// Total number of polygons
+	size_t polygonCount() const { return NumPolys; }
+
+	// Total number of points (vertex locations)
+	size_t pointCount() const { return NumVerts; }
+
+	// Vertex count for polygon n: currently FMeshDescription is just triangles.
+	size_t vertexCount(size_t n) const { return 3; }
+
+	// Return position pos in local grid index space for polygon n and vertex v
+	void getIndexSpacePoint(size_t FaceNumber, size_t CornerNumber, openvdb::Vec3d& pos) const
+	{
+		int32 AdapterIdx = PolyIdxToAdapter[FaceNumber];
+		int32 Offset = AdapterToOffset[AdapterIdx];
+
+		Adapters[AdapterIdx].getIndexSpacePoint(FaceNumber - Offset, CornerNumber, pos);
+	}
+
+	const OpenVDBTransform& Transform() const { return *Xform; }
+
+private:
+
+	FPolygonSoup();
+
+	OpenVDBTransform::Ptr Xform;
+
+	TArray<FMeshDescriptionAdapter> Adapters;
+	TArray<size_t> AdapterToOffset;
+	TArray<int32> PolyIdxToAdapter;
+	size_t NumPolys;
+	size_t NumVerts;
+
+};
+
+FPolygonSoup::FPolygonSoup(const TArray<FMeshDescriptionAdapter>& AdapterArray, double VoxelSize)
+{
+
+	Xform = OpenVDBTransform::createLinearTransform(VoxelSize);
+
+	int32 NumMeshes = AdapterArray.Num();
+	Adapters.Reserve(NumMeshes);
+
+	NumPolys = 0;
+	NumVerts = 0;
+	for (auto& Adapter : AdapterArray)
+	{
+		Adapters.Add(Adapter);
+		NumPolys += Adapter.polygonCount();
+		NumVerts += Adapter.pointCount();
+	}
+
+	AdapterToOffset.AddZeroed(NumMeshes + 1);
+	for (int32 i = 1; i < NumMeshes + 1; ++i)
+	{
+		AdapterToOffset[i] = AdapterToOffset[i-1] + AdapterArray[i - 1].polygonCount();
+	}
+
+	PolyIdxToAdapter.Reserve(NumPolys);
+	for (int32 i = 0; i < NumMeshes; ++i)
+	{
+		int32 np = AdapterArray[i].polygonCount();
+		for (int32 j = 0; j < np; ++j)
+		{
+			PolyIdxToAdapter.Add(i);
+		}
+	}
+}
+
+
+
+
+
+class FVoxelBasedCSGImpl : public IVoxelBasedCSG
+{
+public:
+	FVoxelBasedCSGImpl() 
+	{
+		openvdb::initialize();
+
+		double VoxelSize = 0.1;
+		XForm = OpenVDBTransform::createLinearTransform(VoxelSize);
+	}
+
+	FVoxelBasedCSGImpl(double VoxelSize)
+	{
+		openvdb::initialize();
+
+		XForm = OpenVDBTransform::createLinearTransform(VoxelSize);
+	}
+
+	virtual ~FVoxelBasedCSGImpl()
+	{
+		XForm.reset();
+	}
+
+	double GetVoxelSize() const override
+	{
+		return XForm->voxelSize()[0];
+	}
+
+	void SetVoxelSize(double VoxelSize) override 
+	{
+		XForm = OpenVDBTransform::createLinearTransform(VoxelSize);
+	}
+
+	// Note this will become very slow if the isosurface is more than 3 voxel widths from 0.
+	// the tool that calls this should clamp the isosurface.
+	FVector ComputeUnion(const TArray<IVoxelBasedCSG::FPlacedMesh>& PlacedMeshArray, FMeshDescription& ResultMesh, double Adaptivity, double IsoSurface ) const override
+	{
+
+		
+		// convert IsoSurface units to voxels
+
+		const double VoxelSize = GetVoxelSize();
+
+		const double IsoSurfaceInVoxels = IsoSurface / VoxelSize;
+
+		const double ExteriorVoxelWidth = FMath::Max( 2., IsoSurfaceInVoxels + 1.);
+		const double InteriorVoxelWidth = -FMath::Min(-2., IsoSurfaceInVoxels - 1);
+
+	
+		const int32 NumMeshes = PlacedMeshArray.Num();
+		if (NumMeshes == 0)
+		{
+			return FVector(0.f, 0.f, 0.f);
+		}
+
+		// Find the average translation of all the meshes.
+		
+		FVector AverageTranslation(0.f, 0.f, 0.f);
+		for (int32 i = 0; i < NumMeshes; ++i)
+		{
+			AverageTranslation += PlacedMeshArray[i].Transform.GetTranslation();
+		}
+		AverageTranslation /= (float)NumMeshes;
+
+
+		// Target grid to hold the union in SDF form.
+		//openvdb::FloatGrid::Ptr SDFUnionVolume = openvdb::FloatGrid::create(ExteriorVoxelWidth /*background value */);
+		//SDFUnionVolume->setTransform(XForm);
+		
+		// Fill the SDFUnionVolume
+		FMatrix LocalToVoxel = FMatrix::Identity;
+		LocalToVoxel.M[0][0] = VoxelSize;
+		LocalToVoxel.M[1][1] = VoxelSize;
+		LocalToVoxel.M[2][2] = VoxelSize;
+
+		TArray<FMeshDescriptionAdapter> Adapters;
+		for (int32 i = 0, I = PlacedMeshArray.Num(); i < I; ++i)
+		{
+			const FPlacedMesh& PlacedMesh  = PlacedMeshArray[i];
+			const FMeshDescription* MeshPtr = PlacedMesh.Mesh;
+			if (MeshPtr)
+			{
+
+				// Get the transform relative to the average
+				FTransform MeshTransform = PlacedMesh.Transform;
+				MeshTransform.AddToTranslation(-AverageTranslation);
+				FMatrix TransformMatrix = MeshTransform.ToMatrixWithScale().Inverse();
+				
+
+ 				TransformMatrix = LocalToVoxel * TransformMatrix;
+				float* data = &TransformMatrix.M[0][0];
+				openvdb::math::Mat4<float> VDBMatFloat(data);
+
+				openvdb::Mat4R VDBMatDouble(VDBMatFloat);
+				OpenVDBTransform::Ptr LocalXForm = OpenVDBTransform::createLinearTransform(VDBMatDouble);
+
+				// Create a wrapper with OpenVDB semantics. 
+				
+				FMeshDescriptionAdapter MeshAdapter(*MeshPtr, *LocalXForm);
+
+				Adapters.Add(MeshAdapter);
+
+			}
+		}
+
+		FPolygonSoup PolySoup(Adapters, VoxelSize);
+		openvdb::FloatGrid::Ptr SDFUnionVolume = openvdb::tools::meshToVolume<openvdb::FloatGrid>(PolySoup, PolySoup.Transform(), ExteriorVoxelWidth, InteriorVoxelWidth);
+
+
+		// Convert the SDFUnionVolume to a mesh
+		
+		ConvertSDFToMesh(SDFUnionVolume, Adaptivity, IsoSurface, ResultMesh);
+
+		return AverageTranslation;
+	}
+
+	FVector ComputeDifference(const FPlacedMesh& PlacedMeshA, const FPlacedMesh& PlacedMeshB, FMeshDescription& ResultMesh,  double Adaptivity, double IsoSurface) const override
+	{
+
+		// make SDFs
+		openvdb::FloatGrid::Ptr SDFVolumeA;
+		openvdb::FloatGrid::Ptr SDFVolumeB;
+
+		FVector AverageTranslation = GenerateVolumes(IsoSurface, PlacedMeshA, PlacedMeshB, SDFVolumeA, SDFVolumeB);
+
+		// create the difference - result stored in volume A
+		openvdb::tools::csgDifference(*SDFVolumeA, *SDFVolumeB);
+
+		// convert the result
+		ConvertSDFToMesh(SDFVolumeA, Adaptivity, IsoSurface, ResultMesh);
+
+		return AverageTranslation;
+	}
+
+	FVector ComputeIntersection(const FPlacedMesh& PlacedMeshA, const FPlacedMesh& PlacedMeshB, FMeshDescription& ResultMesh, double Adaptivity, double IsoSurface) const override
+	{
+
+		// make SDFs
+		openvdb::FloatGrid::Ptr SDFVolumeA;
+		openvdb::FloatGrid::Ptr SDFVolumeB;
+
+		FVector AverageTranslation = GenerateVolumes(IsoSurface, PlacedMeshA, PlacedMeshB, SDFVolumeA, SDFVolumeB);
+
+
+		// create the difference - result stored in volume A
+		openvdb::tools::csgIntersection(*SDFVolumeA, *SDFVolumeB);
+
+		// convert the result
+		ConvertSDFToMesh(SDFVolumeA, Adaptivity, IsoSurface, ResultMesh);
+
+		return AverageTranslation;
+	}
+
+	FVector ComputeUnion(const FPlacedMesh& PlacedMeshA, const FPlacedMesh& PlacedMeshB, FMeshDescription& ResultMesh, double Adaptivity, double IsoSurface) const override
+	{
+
+		// make SDFs
+		openvdb::FloatGrid::Ptr SDFVolumeA;
+		openvdb::FloatGrid::Ptr SDFVolumeB;
+
+		FVector AverageTranslation = GenerateVolumes(IsoSurface, PlacedMeshA, PlacedMeshB, SDFVolumeA, SDFVolumeB);
+
+
+		// create the difference - result stored in volume A
+		openvdb::tools::csgUnion(*SDFVolumeA, *SDFVolumeB);
+
+		// convert the result
+		ConvertSDFToMesh(SDFVolumeA, Adaptivity, IsoSurface, ResultMesh);
+
+		return AverageTranslation;
+	}
+
+private:
+
+	void ConvertSDFToMesh(openvdb::FloatGrid::Ptr SDFVolume, double Adaptivity, double IsoSurface, FMeshDescription& ResultMesh) const
+	{
+		// Convert the SDFUnionVolume to a mesh
+		FAOSMesh AOSMeshedVolume;
+		ProxyLOD::ExtractIsosurfaceWithNormals(SDFVolume, IsoSurface, Adaptivity, AOSMeshedVolume); // QUestion should IsoSurface be worldspace?
+		
+																									
+		// Evidently this conversion code makes certain assumptions about the FMeshDescription. that were introduced when it was converted from FRawMesh...
+		
+		ProxyLOD::ConvertMesh(AOSMeshedVolume, ResultMesh);
+	}
+
+	FVector GenerateVolumes(const double IsoSurface, const FPlacedMesh& PlacedMeshA, const FPlacedMesh& PlacedMeshB, openvdb::FloatGrid::Ptr& VolumeA, openvdb::FloatGrid::Ptr& VolumeB) const 
+	{
+		const FMeshDescription& MeshA = *PlacedMeshA.Mesh;
+		const FMeshDescription& MeshB = *PlacedMeshB.Mesh;
+
+		const double VoxelSize = GetVoxelSize();
+
+		const double IsoSurfaceInVoxels = IsoSurface / VoxelSize;
+
+
+		const double ExteriorVoxelWidth = FMath::Max(2., IsoSurfaceInVoxels + 1.);
+		const double InteriorVoxelWidth = -FMath::Min(-2., IsoSurfaceInVoxels - 1);
+
+		// The average translation of the two meshes.
+		FVector AverageTranslation = PlacedMeshA.Transform.GetTranslation() + PlacedMeshB.Transform.GetTranslation();
+		AverageTranslation *= 0.5f;
+
+		// Fill the SDFUnionVolume
+		FMatrix LocalToVoxel = FMatrix::Identity;
+		LocalToVoxel.M[0][0] = VoxelSize;
+		LocalToVoxel.M[1][1] = VoxelSize;
+		LocalToVoxel.M[2][2] = VoxelSize;
+
+		auto TransformGenerator = [&LocalToVoxel, &AverageTranslation](const FPlacedMesh& PlacedMesh)->openvdb::Mat4R
+		{
+			FTransform XForm = PlacedMesh.Transform;
+			XForm.AddToTranslation(-AverageTranslation);
+			FMatrix TransformMatrix = XForm.ToMatrixWithScale().Inverse();
+
+			TransformMatrix = LocalToVoxel * TransformMatrix;
+			float* data = &TransformMatrix.M[0][0];
+			openvdb::math::Mat4<float> VDBMatFloat(data);
+			return openvdb::Mat4R(VDBMatFloat);
+		};
+
+		openvdb::Mat4R XFormA = TransformGenerator(PlacedMeshA);
+		OpenVDBTransform::Ptr VDBXFormA = OpenVDBTransform::createLinearTransform(XFormA);
+
+		openvdb::Mat4R XFormB = TransformGenerator(PlacedMeshB);
+		OpenVDBTransform::Ptr VDBXFormB = OpenVDBTransform::createLinearTransform(XFormB);
+
+		// Create adapters that understand the openVDB semantics.
+		FMeshDescriptionAdapter AdapterA(MeshA, *VDBXFormA);
+		FMeshDescriptionAdapter AdapterB(MeshB, *VDBXFormB);
+
+		// target transform
+		OpenVDBTransform::Ptr TargetXForm = OpenVDBTransform::createLinearTransform(VoxelSize);
+
+		// make SDFs
+		openvdb::FloatGrid::Ptr SDFVolumeA = openvdb::tools::meshToVolume<openvdb::FloatGrid>(AdapterA, *TargetXForm, ExteriorVoxelWidth, InteriorVoxelWidth);
+		openvdb::FloatGrid::Ptr SDFVolumeB = openvdb::tools::meshToVolume<openvdb::FloatGrid>(AdapterB, *TargetXForm, ExteriorVoxelWidth, InteriorVoxelWidth);
+
+		VolumeA = SDFVolumeA;
+		VolumeB = SDFVolumeB;
+
+		return AverageTranslation;
+
+	}
+
+	OpenVDBTransform::Ptr XForm;
+};
+
+TUniquePtr<IVoxelBasedCSG> IVoxelBasedCSG::CreateCSGTool(float VoxelSize)
+{
+	TUniquePtr<FVoxelBasedCSGImpl> CSGTool = MakeUnique<FVoxelBasedCSGImpl>();
+
+	if (CSGTool == nullptr)
+	{
+		return nullptr;
+	}
+	else
+	{
+		CSGTool->SetVoxelSize(VoxelSize);
+	}
+
+	return CSGTool;
+}
\ No newline at end of file
diff --git a/Engine/Plugins/Experimental/ProxyLODPlugin/Source/ProxyLOD/Public/ProxyLODVolume.h b/Engine/Plugins/Experimental/ProxyLODPlugin/Source/ProxyLOD/Public/ProxyLODVolume.h
--- a/Engine/Plugins/Experimental/ProxyLODPlugin/Source/ProxyLOD/Public/ProxyLODVolume.h
+++ b/Engine/Plugins/Experimental/ProxyLODPlugin/Source/ProxyLOD/Public/ProxyLODVolume.h
@@ -56,3 +56,46 @@ public:
 	*/
 	virtual float QueryDistance(const FVector& Point) const = 0;
 };
+
+
+
+
+class PROXYLODMESHREDUCTION_API IVoxelBasedCSG
+{
+public :
+	static TUniquePtr<IVoxelBasedCSG> CreateCSGTool(float VoxelSize);
+
+	class FPlacedMesh
+	{
+	public:
+		FPlacedMesh()
+			: Mesh(nullptr)
+		{}
+
+		FPlacedMesh(const FPlacedMesh& other)
+			: Mesh(other.Mesh)
+			, Transform(other.Transform)
+		{}
+
+		FPlacedMesh& operator=(const FPlacedMesh& other)
+		{
+			Mesh = other.Mesh;
+			Transform = other.Transform;
+			return *this;
+		}
+
+		const FMeshDescription* Mesh;
+		FTransform       Transform;
+	};
+
+	virtual double GetVoxelSize() const = 0;
+	virtual void SetVoxelSize(double VolexSize) = 0;
+
+	// Will destroy UVs and other attributes  Returns the average location of the input meshes 
+	virtual FVector ComputeUnion(const TArray<IVoxelBasedCSG::FPlacedMesh>& MeshArray, FMeshDescription& ResultMesh, double Adaptivity = 0.1, double IsoSurfcae = 0.) const = 0;
+	
+	// We could make this keep the UVs and other attributes from the AMesh..
+	virtual FVector ComputeDifference(const FPlacedMesh& PlacedMeshA, const FPlacedMesh& PlacedMeshB, FMeshDescription& ResultMesh, double Adaptivity, double IsoSurface) const = 0;
+	virtual FVector ComputeIntersection(const FPlacedMesh& PlacedMeshA, const FPlacedMesh& PlacedMeshB, FMeshDescription& ResultMesh, double Adaptivity, double IsoSurface) const = 0;
+	virtual FVector ComputeUnion(const FPlacedMesh& PlacedMeshA, const FPlacedMesh& PlacedMeshB, FMeshDescription& ResultMesh, double Adaptivity, double IsoSurface) const = 0;
+};
\ No newline at end of file
