copy up CL 6664465 ProxyLOD - Expose mesh UV parameterization tools for use elsewhere
#rnx

diff --git a/Engine/Plugins/Experimental/ProxyLODPlugin/Source/ProxyLOD/Private/ProxLODParameterization.cpp b/Engine/Plugins/Experimental/ProxyLODPlugin/Source/ProxyLOD/Private/ProxLODParameterization.cpp
new file mode 100644
--- /dev/null
+++ b/Engine/Plugins/Experimental/ProxyLODPlugin/Source/ProxyLOD/Private/ProxLODParameterization.cpp
@@ -0,0 +1,92 @@
+// Copyright 1998-2019 Epic Games, Inc. All Rights Reserved.
+
+#include "ProxyLODParameterization.h"
+
+#include "ProxyLODMeshConvertUtils.h"
+#include "ProxyLODMeshUtilities.h"
+#include "ProxyLODMeshParameterization.h"
+
+
+
+
+class FPoxyLODParameterizationImpl : public IPoxyLODParameterization
+{
+public:
+	FPoxyLODParameterizationImpl() {}
+
+	virtual ~FPoxyLODParameterizationImpl() {};
+
+
+	virtual bool ParameterizeMeshDescription(FMeshDescription& MeshDescription, int32 Width, int32 Height, float GutterSpace, float Stretch, int32 ChartNum, bool bUseNormals, bool bRecomputeTangentSpace, bool bPrintDebugMessages) const override
+	{
+		const float MaxStretch = Stretch;
+		const size_t MaxChartNumber = ChartNum;
+		const bool bUseNormalsInUV = bUseNormals;
+		
+		// No Op
+		auto NoOpCallBack = [](float percent)->HRESULT {return  S_OK; };
+
+		bool bComputeNormals = bRecomputeTangentSpace;
+		bool bSplitHardAngles = false;
+		bool bComputeTangentSpace = bRecomputeTangentSpace;
+
+		float HardAngleDegrees = 80.f;
+		FVertexDataMesh VertexDataMesh;
+		ProxyLOD::ConvertMesh(MeshDescription, VertexDataMesh);
+
+		if (bSplitHardAngles)
+		{
+			ProxyLOD::SplitHardAngles(HardAngleDegrees, VertexDataMesh);
+		}
+
+		if (bComputeNormals)
+		{
+			ProxyLOD::ComputeVertexNormals(VertexDataMesh, ProxyLOD::ENormalComputationMethod::AngleWeighted);
+		}
+
+
+		FIntPoint UVSize(Height, Width);
+		ProxyLOD::FTextureAtlasDesc TextureAtlasDesc(UVSize, GutterSpace);
+
+		float MaxStretchUsed = 0.f;
+		size_t NumChartsMade = 0;
+		bool bUVGenerationSucess = ProxyLOD::GenerateUVs(VertexDataMesh, TextureAtlasDesc, false, MaxStretch, MaxChartNumber, bUseNormalsInUV, NoOpCallBack, &MaxStretchUsed, &NumChartsMade);
+
+		if (bUVGenerationSucess)
+		{
+			if (bPrintDebugMessages)
+			{
+				UE_LOG(LogTemp, Warning, TEXT("UV Generation: generated %d charts, using max stretch  %f "), NumChartsMade, MaxStretchUsed);
+			}
+
+			if (bComputeTangentSpace)
+			{
+				ProxyLOD::ComputeTangentSpace(VertexDataMesh, false);
+			}
+			// Reset the result mesh description
+			MeshDescription.Empty();
+			UStaticMesh::RegisterMeshAttributes(MeshDescription);
+
+			// convert to fill the mesh description
+			ProxyLOD::ConvertMesh(VertexDataMesh, MeshDescription);
+		}
+
+
+		return bUVGenerationSucess;
+	};
+};
+
+
+
+TUniquePtr<IPoxyLODParameterization> IPoxyLODParameterization::CreateTool()
+{
+	TUniquePtr<IPoxyLODParameterization> Tool = MakeUnique<FPoxyLODParameterizationImpl>();
+
+	if (Tool == nullptr)
+	{
+		return nullptr;
+	}
+
+	return Tool;
+
+}
\ No newline at end of file
diff --git a/Engine/Plugins/Experimental/ProxyLODPlugin/Source/ProxyLOD/Private/ProxyLODMeshParameterization.cpp b/Engine/Plugins/Experimental/ProxyLODPlugin/Source/ProxyLOD/Private/ProxyLODMeshParameterization.cpp
--- a/Engine/Plugins/Experimental/ProxyLODPlugin/Source/ProxyLOD/Private/ProxyLODMeshParameterization.cpp
+++ b/Engine/Plugins/Experimental/ProxyLODPlugin/Source/ProxyLOD/Private/ProxyLODMeshParameterization.cpp
@@ -11,21 +11,39 @@ THIRD_PARTY_INCLUDES_END
 #include "ProxyLODThreadedWrappers.h"
 #include "ProxyLODMeshUtilities.h"
 
-
 bool ProxyLOD::GenerateUVs(FVertexDataMesh& InOutMesh, const FTextureAtlasDesc& TextureAtlasDesc, const bool VertexColorParts)
 {
+	// desired parameters for ISO-Chart method
 
-	std::vector<uint32> DirextXAdjacency;
+	// MaxChartNum = 0 will allow any number of charts to be generated.
 
-	const bool bValidMesh = GenerateAdjacenyAndCleanMesh(InOutMesh, DirextXAdjacency);
+	const size_t MaxChartNumber = 0;
 
-	if (!bValidMesh) return false;
+	// Let the polys in the partitions stretch some..  1.f will let it stretch freely
 
+	const float MaxStretch = 0.125f; // Question.  Does this override the pIMT?
+	
+	// Note: I tried constructing this from the normals, but that resulted in some large planar regions being really compressed in the
+	// UV chart.
+	const bool bComputeIMTFromVertexNormal = false;
 
 	// No Op
 	auto NoOpCallBack = [](float percent)->HRESULT {return  S_OK; };
 
-	// Gather data for UVAtlas Create
+	return GenerateUVs(InOutMesh, TextureAtlasDesc, VertexColorParts, MaxStretch, MaxChartNumber, bComputeIMTFromVertexNormal, NoOpCallBack);
+}
+
+bool ProxyLOD::GenerateUVs(FVertexDataMesh& InOutMesh, const FTextureAtlasDesc& TextureAtlasDesc, const bool VertexColorParts, 
+	                      const float MaxStretch, const size_t MaxChartNumber, const bool bComputeIMTFromVertexNormal, 
+	                      std::function<HRESULT __cdecl(float percentComplete)> StatusCallBack, float* MaxStretchOut, size_t* NumChartsOut)
+{
+
+	std::vector<uint32> DirextXAdjacency;
+
+	const bool bValidMesh = GenerateAdjacenyAndCleanMesh(InOutMesh, DirextXAdjacency);
+
+	if (!bValidMesh) return false;
+
 
 	// Data from the existing mesh
 
@@ -40,16 +58,6 @@ bool ProxyLOD::GenerateUVs(FVertexDataMesh& InOutMesh, const FTextureAtlasDesc&
 
 	const uint32* adjacency = DirextXAdjacency.data();
 
-	// desired parameters for ISO-Chart method
-
-	// MaxChartNum = 0 will allow any number of charts to be generated.
-
-	const size_t MaxChartNumber = 0;
-
-	// Let the polys in the partitions stretch some..  1.f will let it stretch freely
-
-	const float MaxStretch = 0.125f; // Question.  Does this override the pIMT?
-
 
 	// Size of the texture atlas
 
@@ -66,21 +74,36 @@ bool ProxyLOD::GenerateUVs(FVertexDataMesh& InOutMesh, const FTextureAtlasDesc&
 	std::vector<uint32> facePartitioning;
 
 	// Capture stats about the result.
-	float maxStretchOut = 0.f;
-	size_t numChartsOut = 0;
+	float maxStretchUsed = 0.f;
+	size_t numChartsUsed = 0;
 
-	// per-triangle IMT from per-vertex data(normal).  Note: I tried constructing this from
-	// the normals, but that resulted in some large planar regions being really compressed in the
-	// UV chart.
+	
 
 	float * pIMTArray = new float[NumFaces * 3];
-	for (int32 f = 0; f < NumFaces; ++f)
+	if (!bComputeIMTFromVertexNormal)
+	{
+		for (int32 f = 0; f < NumFaces; ++f)
+		{
+			int32 offset = 3 * f;
+			{
+				pIMTArray[offset] = 1.f;
+				pIMTArray[offset + 1] = 0.f;
+				pIMTArray[offset + 2] = 1.f;
+			}
+		}
+	}
+	else
 	{
-		int32 offset = 3 * f;
+		// per-triangle IMT from per-vertex data(normal). 
+
+		const TArray<FVector>& Normals = InOutMesh.Normal;
+		const float* PerVertSignal = (float*)Normals.GetData();
+		size_t SignalStride = 3 * sizeof(float);
+		HRESULT IMTResult = UVAtlasComputeIMTFromPerVertexSignal(Pos, NumVerts, indices, DXGI_FORMAT_R32_UINT, NumFaces, PerVertSignal, 3, SignalStride, StatusCallBack, pIMTArray);
+
+		if (FAILED(IMTResult))
 		{
-			pIMTArray[offset    ] = 1.f;
-			pIMTArray[offset + 1] = 0.f;
-			pIMTArray[offset + 2] = 1.f;
+			return false;
 		}
 	}
 
@@ -89,10 +112,18 @@ bool ProxyLOD::GenerateUVs(FVertexDataMesh& InOutMesh, const FTextureAtlasDesc&
 		MaxChartNumber, MaxStretch,
 		width, height, gutter,
 		adjacency, NULL /*false adj*/, pIMTArray /*IMTArray*/,
-		NoOpCallBack, DirectX::UVATLAS_DEFAULT_CALLBACK_FREQUENCY,
+		StatusCallBack, DirectX::UVATLAS_DEFAULT_CALLBACK_FREQUENCY,
 		DirectX::UVATLAS_DEFAULT, vb, ib,
-		&facePartitioning, &vertexRemapArray, &maxStretchOut, &numChartsOut);
+		&facePartitioning, &vertexRemapArray, &maxStretchUsed, &numChartsUsed);
 
+	if (MaxStretchOut)
+	{
+		*MaxStretchOut = maxStretchUsed;
+	}
+	if (NumChartsOut)
+	{
+		*NumChartsOut = numChartsUsed;
+	}
 	if (FAILED(hr)) return false;
 
 	if (pIMTArray) delete[] pIMTArray;
diff --git a/Engine/Plugins/Experimental/ProxyLODPlugin/Source/ProxyLOD/Private/ProxyLODMeshParameterization.h b/Engine/Plugins/Experimental/ProxyLODPlugin/Source/ProxyLOD/Private/ProxyLODMeshParameterization.h
--- a/Engine/Plugins/Experimental/ProxyLODPlugin/Source/ProxyLOD/Private/ProxyLODMeshParameterization.h
+++ b/Engine/Plugins/Experimental/ProxyLODPlugin/Source/ProxyLOD/Private/ProxyLODMeshParameterization.h
@@ -5,7 +5,7 @@
 #include "ProxyLODMeshTypes.h"  // FVertexDataMesh
 #include "ProxyLODGrid2d.h" // FTextureAtlasDesc
 
-#include <vector>
+#include <functional> // Note: should convert std::function to TFunction
 
 /**
 *  In mesh segmentation (charts and atlases) and parameterization (generating UVs ) we rely 
@@ -33,6 +33,34 @@ namespace ProxyLOD
 	*/
 	bool GenerateUVs( FVertexDataMesh& InOutMesh, const FTextureAtlasDesc& InTextureAtlasDesc, const bool bVertexColorParts = false);
 
+
+
+	/**
+	* Lower-level entry point:
+	* Method that generates new UVs on the FVertexDataMesh according to the parameters specified in the FTextureAtlasDesc
+	* The underlying code uses Isometeric approach (Iso-Charts) in UV generation.
+	*
+	* As a debugging option, the updated InOutMesh can have vertex colors that distinguish the various UV charts.
+	*
+	* NB: The mesh vertex count may change as vertices are split on UV seams.
+	*
+	* @param InOutMesh            Mesh that will be updated by this function, adding UVs
+	* @param InTextureAtlasDesc   Description of the texel resolution of the desired texture atlas.
+	* @param bVertColorParts      Option to add vertex colors according the the chart in the texture atlas, used for debugging.
+	* @param MaxStretch           The maximum amount of stretch between (0 - none and 1 - any)
+	* @param MaxChartNumber       Maximum number of charts required for the atlas.  If this is 0, will be parameterized solely on stretch. 
+	*                             Note, not a hard limit - isochart will stop when a valid charting is found that is greater or equal to this number.
+	* @param MaxStretchOut        Actual max stretch used in computing uvs
+	* @param NumChartsOut         Number of charts actually generated
+	*
+	* @return 'true' if the UV generation succeeded,  'false' if it failed.
+	*/
+	bool GenerateUVs(FVertexDataMesh& InOutMesh, const FTextureAtlasDesc& TextureAtlasDesc, const bool VertexColorParts,
+												 const float MaxStretch, const size_t MaxChartNumber, 
+		                                         const bool bComputeIMTFromVertexNormal,
+		                                         std::function<HRESULT __cdecl(float percentComplete)> StatusCallBack,
+		                                         float* MaxStretchOut = NULL, size_t* NumChartsOut = NULL);
+
 	/**
 	* Generate adjacency data needed for the mesh, additionally this may alter the mesh in attempting to 
 	* remove mesh degeneracy problems.  This method is primarily called within GenerateUVs
diff --git a/Engine/Plugins/Experimental/ProxyLODPlugin/Source/ProxyLOD/Public/ProxyLODParameterization.h b/Engine/Plugins/Experimental/ProxyLODPlugin/Source/ProxyLOD/Public/ProxyLODParameterization.h
new file mode 100644
--- /dev/null
+++ b/Engine/Plugins/Experimental/ProxyLODPlugin/Source/ProxyLOD/Public/ProxyLODParameterization.h
@@ -0,0 +1,18 @@
+// Copyright 1998 - 2019 Epic Games, Inc.All Rights Reserved.
+
+#pragma once
+
+#include "CoreMinimal.h"
+
+struct FMeshDescription;
+
+
+class PROXYLODMESHREDUCTION_API IPoxyLODParameterization
+{
+public:
+	static TUniquePtr<IPoxyLODParameterization> CreateTool();
+
+
+	virtual bool ParameterizeMeshDescription(FMeshDescription& MeshDescription, int32 Width, int32 Height, float GutterSpace, 
+		                                     float Stretch, int32 ChartNum, bool bUseNormals, bool bRecomputeTangentSpace, bool bPrintDebugMessages) const = 0;
+};
\ No newline at end of file
