Fix clang build
#rnx

diff --git a/Engine/Plugins/Experimental/ProxyLODPlugin/Source/ProxyLOD/Public/ProxyLODParameterization.h b/Engine/Plugins/Experimental/ProxyLODPlugin/Source/ProxyLOD/Public/ProxyLODParameterization.h
--- a/Engine/Plugins/Experimental/ProxyLODPlugin/Source/ProxyLOD/Public/ProxyLODParameterization.h
+++ b/Engine/Plugins/Experimental/ProxyLODPlugin/Source/ProxyLOD/Public/ProxyLODParameterization.h
@@ -11,8 +11,7 @@ class PROXYLODMESHREDUCTION_API IPoxyLODParameterization
 {
 public:
 	static TUniquePtr<IPoxyLODParameterization> CreateTool();
-
-
+	virtual ~IPoxyLODParameterization(){}
 	virtual bool ParameterizeMeshDescription(FMeshDescription& MeshDescription, int32 Width, int32 Height, float GutterSpace, 
 		                                     float Stretch, int32 ChartNum, bool bUseNormals, bool bRecomputeTangentSpace, bool bPrintDebugMessages) const = 0;
-};
\ No newline at end of file
+};
