Fix some spelling errors

#rnx

diff --git a/Engine/Plugins/Experimental/ProxyLODPlugin/Source/ProxyLOD/Private/ProxLODParameterization.cpp b/Engine/Plugins/Experimental/ProxyLODPlugin/Source/ProxyLOD/Private/ProxLODParameterization.cpp
deleted file mode 100644
diff --git a/Engine/Plugins/Experimental/ProxyLODPlugin/Source/ProxyLOD/Private/ProxyLODParameterization.cpp b/Engine/Plugins/Experimental/ProxyLODPlugin/Source/ProxyLOD/Private/ProxyLODParameterization.cpp
new file mode 100644
--- /dev/null
+++ b/Engine/Plugins/Experimental/ProxyLODPlugin/Source/ProxyLOD/Private/ProxyLODParameterization.cpp
@@ -0,0 +1,92 @@
+// Copyright 1998-2019 Epic Games, Inc. All Rights Reserved.
+
+#include "ProxyLODParameterization.h"
+
+#include "ProxyLODMeshConvertUtils.h"
+#include "ProxyLODMeshUtilities.h"
+#include "ProxyLODMeshParameterization.h"
+
+
+
+
+class FProxyLODParameterizationImpl : public IProxyLODParameterization
+{
+public:
+	FProxyLODParameterizationImpl() {}
+
+	virtual ~FProxyLODParameterizationImpl() {};
+
+
+	virtual bool ParameterizeMeshDescription(FMeshDescription& MeshDescription, int32 Width, int32 Height, float GutterSpace, float Stretch, int32 ChartNum, bool bUseNormals, bool bRecomputeTangentSpace, bool bPrintDebugMessages) const override
+	{
+		const float MaxStretch = Stretch;
+		const size_t MaxChartNumber = ChartNum;
+		const bool bUseNormalsInUV = bUseNormals;
+		
+		// No Op
+		auto NoOpCallBack = [](float percent)->HRESULT {return  S_OK; };
+
+		bool bComputeNormals = bRecomputeTangentSpace;
+		bool bSplitHardAngles = false;
+		bool bComputeTangentSpace = bRecomputeTangentSpace;
+
+		float HardAngleDegrees = 80.f;
+		FVertexDataMesh VertexDataMesh;
+		ProxyLOD::ConvertMesh(MeshDescription, VertexDataMesh);
+
+		if (bSplitHardAngles)
+		{
+			ProxyLOD::SplitHardAngles(HardAngleDegrees, VertexDataMesh);
+		}
+
+		if (bComputeNormals)
+		{
+			ProxyLOD::ComputeVertexNormals(VertexDataMesh, ProxyLOD::ENormalComputationMethod::AngleWeighted);
+		}
+
+
+		FIntPoint UVSize(Height, Width);
+		ProxyLOD::FTextureAtlasDesc TextureAtlasDesc(UVSize, GutterSpace);
+
+		float MaxStretchUsed = 0.f;
+		size_t NumChartsMade = 0;
+		bool bUVGenerationSucess = ProxyLOD::GenerateUVs(VertexDataMesh, TextureAtlasDesc, false, MaxStretch, MaxChartNumber, bUseNormalsInUV, NoOpCallBack, &MaxStretchUsed, &NumChartsMade);
+
+		if (bUVGenerationSucess)
+		{
+			if (bPrintDebugMessages)
+			{
+				UE_LOG(LogTemp, Warning, TEXT("UV Generation: generated %d charts, using max stretch  %f "), NumChartsMade, MaxStretchUsed);
+			}
+
+			if (bComputeTangentSpace)
+			{
+				ProxyLOD::ComputeTangentSpace(VertexDataMesh, false);
+			}
+			// Reset the result mesh description
+			MeshDescription.Empty();
+			UStaticMesh::RegisterMeshAttributes(MeshDescription);
+
+			// convert to fill the mesh description
+			ProxyLOD::ConvertMesh(VertexDataMesh, MeshDescription);
+		}
+
+
+		return bUVGenerationSucess;
+	};
+};
+
+
+
+TUniquePtr<IProxyLODParameterization> IProxyLODParameterization::CreateTool()
+{
+	TUniquePtr<IProxyLODParameterization> Tool = MakeUnique<FProxyLODParameterizationImpl>();
+
+	if (Tool == nullptr)
+	{
+		return nullptr;
+	}
+
+	return Tool;
+
+}
diff --git a/Engine/Plugins/Experimental/ProxyLODPlugin/Source/ProxyLOD/Public/ProxyLODParameterization.h b/Engine/Plugins/Experimental/ProxyLODPlugin/Source/ProxyLOD/Public/ProxyLODParameterization.h
--- a/Engine/Plugins/Experimental/ProxyLODPlugin/Source/ProxyLOD/Public/ProxyLODParameterization.h
+++ b/Engine/Plugins/Experimental/ProxyLODPlugin/Source/ProxyLOD/Public/ProxyLODParameterization.h
@@ -7,11 +7,12 @@
 struct FMeshDescription;
 
 
-class PROXYLODMESHREDUCTION_API IPoxyLODParameterization
+class PROXYLODMESHREDUCTION_API IProxyLODParameterization
 {
 public:
-	static TUniquePtr<IPoxyLODParameterization> CreateTool();
-	virtual ~IPoxyLODParameterization(){}
-	virtual bool ParameterizeMeshDescription(FMeshDescription& MeshDescription, int32 Width, int32 Height, float GutterSpace, 
-		                                     float Stretch, int32 ChartNum, bool bUseNormals, bool bRecomputeTangentSpace, bool bPrintDebugMessages) const = 0;
+	static TUniquePtr<IProxyLODParameterization> CreateTool();
+	virtual ~IProxyLODParameterization(){}
+	virtual bool ParameterizeMeshDescription(FMeshDescription& MeshDescription, int32 Width, int32 Height,
+											 float GutterSpace, float Stretch, int32 ChartNum,
+											 bool bUseNormals, bool bRecomputeTangentSpace, bool bPrintDebugMessages) const = 0;
 };
