ProxyLOD - Expose lower-level mesh paramterization for use by other tools.

diff --git a/Engine/Plugins/Experimental/ProxyLODPlugin/Source/ProxyLOD/Private/ProxyLODMeshParameterization.cpp b/Engine/Plugins/Experimental/ProxyLODPlugin/Source/ProxyLOD/Private/ProxyLODMeshParameterization.cpp
--- a/Engine/Plugins/Experimental/ProxyLODPlugin/Source/ProxyLOD/Private/ProxyLODMeshParameterization.cpp
+++ b/Engine/Plugins/Experimental/ProxyLODPlugin/Source/ProxyLOD/Private/ProxyLODMeshParameterization.cpp
@@ -1,6 +1,7 @@
 // Copyright 1998-2019 Epic Games, Inc. All Rights Reserved.
 
 #include "ProxyLODMeshParameterization.h"
+#include "CoreMinimal.h"
 
 THIRD_PARTY_INCLUDES_START
 #include <UVAtlasCode/UVAtlas/inc/UVAtlas.h>
@@ -11,6 +12,157 @@ THIRD_PARTY_INCLUDES_END
 #include "ProxyLODThreadedWrappers.h"
 #include "ProxyLODMeshUtilities.h"
 
+
+bool ProxyLOD::GenerateUVs(const FTextureAtlasDesc& TextureAtlasDesc,
+	                       const TArray<FVector>&   VertexBuffer, 
+	                       const TArray<int32>&     IndexBuffer, 
+						   const TArray<int32>&     AdjacencyBuffer,
+	                       TFunction<bool(float)>&  Callback, 
+	                       TArray<FVector2D>&       UVVertexBuffer, 
+	                       TArray<int32>&           UVIndexBuffer, 
+	                       TArray<int32>&           VertexRemapArray, 
+	                       float&                   MaxStretch, 
+	                       int32&                   NumCharts)
+{
+	const int32 NumVerts   = VertexBuffer.Num();
+	const int32 NumIndices = IndexBuffer.Num();
+	const int32 NumFaces   = IndexBuffer.Num() / 3;
+
+	// Copy data into DirectX library format
+	DirectX::XMFLOAT3* PosArray = new DirectX::XMFLOAT3[NumVerts];
+	uint32* Indices             = new uint32[NumIndices];
+	uint32* AdjacencyArray      = new uint32[3 * NumFaces];
+	{
+		for (int32 i = 0; i < NumVerts; ++i)
+		{
+			const FVector& Vertex = VertexBuffer[i];
+			PosArray[i] = DirectX::XMFLOAT3(Vertex.X, Vertex.Y, Vertex.Z);
+		}
+		for (int32 i = 0; i < NumIndices; ++i)
+		{
+			Indices[i] = (uint32)IndexBuffer[i];
+		}
+		for (int32 i = 0, I = NumFaces * 3; i <  I; ++i)
+		{
+			AdjacencyArray[i] = static_cast<uint32>(AdjacencyBuffer[i]);
+		}
+
+	}
+
+
+	// Verify the mesh is valid
+	{
+		HRESULT ValidateHR = DirectX::Validate(Indices, NumFaces, NumVerts, AdjacencyArray, DirectX::VALIDATE_BOWTIES, NULL);
+		if (FAILED(ValidateHR))
+		{
+			return false;
+		}
+	}
+
+	// Capture the callback and convert result type
+	auto StatusCallBack = [&Callback](float percentComplete)->HRESULT
+	{
+		//return S_OK;
+		return (Callback(percentComplete)) ? S_OK : S_FALSE;
+	};
+
+	
+
+	// Translate controls into corret types
+	const size_t NumChartsIn  = (size_t)NumCharts;
+	const float  MaxStretchIn = MaxStretch;
+
+	const size_t Width  = TextureAtlasDesc.Size.X;
+	const size_t Height = TextureAtlasDesc.Size.Y;
+	const float  Gutter = TextureAtlasDesc.Gutter;
+
+	// Symmetric Identity matrix used for the signal
+	float * pIMTArray = new float[NumFaces * 3];
+	{
+		for (int32 f = 0; f < NumFaces; ++f)
+		{
+			int32 offset = 3 * f;
+			{
+				pIMTArray[offset + 0] = 1.f;
+				pIMTArray[offset + 1] = 0.f;
+				pIMTArray[offset + 2] = 1.f;
+			}
+		}
+	}
+
+
+	size_t NumChartsOut = 0;
+	float MaxStretchOut = 0.f;
+	
+	// info to capture
+	std::vector<DirectX::UVAtlasVertex> VB;
+	std::vector<uint8>  IB;
+	std::vector<uint32> RemapArray;
+	std::vector<uint32> FacePartitioning;
+
+	// Generate UVs
+	HRESULT hr = DirectX::UVAtlasCreate(PosArray, NumVerts,
+		                                Indices, DXGI_FORMAT_R32_UINT, NumFaces,
+		                                NumChartsIn, MaxStretchIn,
+		                                Width, Height, Gutter,
+		                                AdjacencyArray, NULL /*false adj*/, pIMTArray /*IMTArray*/,
+		                                StatusCallBack, DirectX::UVATLAS_DEFAULT_CALLBACK_FREQUENCY,
+		                                DirectX::UVATLAS_DEFAULT, VB, IB,
+		                                &FacePartitioning, &RemapArray, &MaxStretchOut, &NumChartsOut);
+
+
+	// Translate results to output form.
+	if (hr == S_OK)
+	{
+		const int32 NumUVverts    = (int32)VB.size();
+		const int32 NumUVindices  = (int32)IB.size();
+		const int32 NumRemapArray = (int32)RemapArray.size();
+
+		// empty and resize
+
+		UVVertexBuffer.Empty(NumUVverts);
+		UVIndexBuffer.Empty(NumUVindices);
+		VertexRemapArray.Empty(NumRemapArray);
+
+		for (const DirectX::UVAtlasVertex& UVvertex : VB)
+		{
+			const auto& UV = UVvertex.uv;
+			UVVertexBuffer.Emplace(FVector2D(UV.x, UV.y));
+		}
+
+		// This part is weird
+		{
+			std::vector<uint32> indices;
+			indices.resize(3 * NumFaces);
+			std::memcpy(indices.data(), IB.data(), sizeof(uint32) * 3 * NumFaces);
+
+			for (uint32 i : indices)
+			{
+				UVIndexBuffer.Add(i);
+			}
+		}
+
+		for (const uint32 i : RemapArray)
+		{
+			int32 r = i;
+			VertexRemapArray.Add(r);
+		}
+
+		MaxStretch = MaxStretchOut;
+		NumCharts  = (int32)NumChartsOut;
+	}
+
+	// Clean up
+	delete[] pIMTArray;
+	delete[] PosArray;
+	delete[] Indices;
+	delete[] AdjacencyArray;
+	
+
+	return (hr == S_OK) ? true : false;
+}
+
+
 bool ProxyLOD::GenerateUVs(FVertexDataMesh& InOutMesh, const FTextureAtlasDesc& TextureAtlasDesc, const bool VertexColorParts)
 {
 	// desired parameters for ISO-Chart method
diff --git a/Engine/Plugins/Experimental/ProxyLODPlugin/Source/ProxyLOD/Private/ProxyLODMeshParameterization.h b/Engine/Plugins/Experimental/ProxyLODPlugin/Source/ProxyLOD/Private/ProxyLODMeshParameterization.h
--- a/Engine/Plugins/Experimental/ProxyLODPlugin/Source/ProxyLOD/Private/ProxyLODMeshParameterization.h
+++ b/Engine/Plugins/Experimental/ProxyLODPlugin/Source/ProxyLOD/Private/ProxyLODMeshParameterization.h
@@ -50,16 +50,58 @@ namespace ProxyLOD
 	* @param MaxStretch           The maximum amount of stretch between (0 - none and 1 - any)
 	* @param MaxChartNumber       Maximum number of charts required for the atlas.  If this is 0, will be parameterized solely on stretch. 
 	*                             Note, not a hard limit - isochart will stop when a valid charting is found that is greater or equal to this number.
+	* @param bComputeIMTNormal    Compute a metric tensor from the normals.  Generally not a good thing.
+	* @param StatusCallBack       Allows for termination of the UV generation.
 	* @param MaxStretchOut        Actual max stretch used in computing uvs
 	* @param NumChartsOut         Number of charts actually generated
 	*
 	* @return 'true' if the UV generation succeeded,  'false' if it failed.
 	*/
-	bool GenerateUVs(FVertexDataMesh& InOutMesh, const FTextureAtlasDesc& TextureAtlasDesc, const bool VertexColorParts,
-												 const float MaxStretch, const size_t MaxChartNumber, 
-		                                         const bool bComputeIMTFromVertexNormal,
-		                                         std::function<HRESULT __cdecl(float percentComplete)> StatusCallBack,
-		                                         float* MaxStretchOut = NULL, size_t* NumChartsOut = NULL);
+	bool GenerateUVs(FVertexDataMesh& InOutMesh, 
+		             const FTextureAtlasDesc& TextureAtlasDesc, 
+		             const bool bVertexColorParts,
+				     const float MaxStretch, 
+		             const size_t MaxChartNumber, 
+		             const bool bComputeIMTNormal,
+		             std::function<HRESULT __cdecl(float percentComplete)> StatusCallBack,
+		             float* MaxStretchOut = NULL, 
+		             size_t* NumChartsOut = NULL);
+
+
+	/**
+	* Lower-level entry point:
+	* Method that generates new UVs inconstant with the mesh defined by the VertexBuffer/IndexBuffer 
+	* according to the parameters specified in the FTextureAtlasDesc
+	* The underlying code uses Isometeric approach (Iso-Charts) in UV generation.
+	* This assumes that the caller has already generated the mesh adjacency information. 
+	*
+	*
+	*
+	* @param InTextureAtlasDesc   Description of the texel resolution of the desired texture atlas.
+	* @param VertexBuffer         Positions of vertices
+	* @param IndexBuffer          Triangle definitions
+	* @param AdjacencyBuffer      3 entries for each face, each entry is the face index of the face that is adjacent to the given face. -1 means no face.
+	* @param StatusCallBack       Allows for termination of the UV generation.
+	* @param UVVertexBuffer       Resulting UVs
+	* @param UVIndexBuffer        Connectivity for the UVs
+	* @param VertexRemapArray     Array that remaps the vertices split by the UV generation to the source vertices.
+	* @param MaxStretch           On Input, The maximum amount of stretch between (0 - none and 1 - any), on Output the actual max stretch used.
+	* @param MaxChartNumber       On Input Maximum number of charts required for the atlas.  If this is 0, will be parameterized solely on stretch.
+	*                             Note, not a hard limit - isochart will stop when a valid charting is found that is greater or equal to this number.
+	*                             On Output: Number of charts actually generated
+	*
+	* @return 'true' if the UV generation succeeded,  'false' if it failed.
+	*/
+	bool GenerateUVs(const FTextureAtlasDesc& TextureAtlasDesc,
+		             const TArray<FVector>&   VertexBuffer,
+		             const TArray<int32>&     IndexBuffer,
+		             const TArray<int32>&     AdjacencyBuffer,
+		             TFunction<bool(float)>&  Callback,
+		             TArray<FVector2D>&       UVVertexBuffer,
+		             TArray<int32>&           UVIndexBuffer,
+		             TArray<int32>&           VertexRemapArray,
+		             float&                   MaxStretch,
+		             int32&                   NumCharts);
 
 	/**
 	* Generate adjacency data needed for the mesh, additionally this may alter the mesh in attempting to 
diff --git a/Engine/Plugins/Experimental/ProxyLODPlugin/Source/ProxyLOD/Private/ProxyLODParameterization.cpp b/Engine/Plugins/Experimental/ProxyLODPlugin/Source/ProxyLOD/Private/ProxyLODParameterization.cpp
--- a/Engine/Plugins/Experimental/ProxyLODPlugin/Source/ProxyLOD/Private/ProxyLODParameterization.cpp
+++ b/Engine/Plugins/Experimental/ProxyLODPlugin/Source/ProxyLOD/Private/ProxyLODParameterization.cpp
@@ -74,6 +74,24 @@ public:
 
 		return bUVGenerationSucess;
 	};
+
+
+	virtual bool GenerateUVs(int32                    Width,
+		                     int32                    Height,
+		                     float                    GutterSpace,
+			                 const TArray<FVector>&   VertexBuffer,
+			                 const TArray<int32>&     IndexBuffer,
+			                 const TArray<int32>&     AdjacencyBuffer,
+			                 TFunction<bool(float)>&  Callback,
+			                 TArray<FVector2D>&       UVVertexBuffer,
+			                 TArray<int32>&           UVIndexBuffer,
+			                 TArray<int32>&           VertexRemapArray,
+			                 float&                   MaxStretch,
+			                 int32&                   NumCharts) const 
+	{
+		ProxyLOD::FTextureAtlasDesc TextureAtlasDesc(FIntPoint(Height, Width), GutterSpace);
+		return ProxyLOD::GenerateUVs(TextureAtlasDesc, VertexBuffer, IndexBuffer, AdjacencyBuffer, Callback, UVVertexBuffer, UVIndexBuffer, VertexRemapArray, MaxStretch, NumCharts);
+	}
 };
 
 
diff --git a/Engine/Plugins/Experimental/ProxyLODPlugin/Source/ProxyLOD/Public/ProxyLODParameterization.h b/Engine/Plugins/Experimental/ProxyLODPlugin/Source/ProxyLOD/Public/ProxyLODParameterization.h
--- a/Engine/Plugins/Experimental/ProxyLODPlugin/Source/ProxyLOD/Public/ProxyLODParameterization.h
+++ b/Engine/Plugins/Experimental/ProxyLODPlugin/Source/ProxyLOD/Public/ProxyLODParameterization.h
@@ -12,7 +12,51 @@ class PROXYLODMESHREDUCTION_API IProxyLODParameterization
 public:
 	static TUniquePtr<IProxyLODParameterization> CreateTool();
 	virtual ~IProxyLODParameterization(){}
-	virtual bool ParameterizeMeshDescription(FMeshDescription& MeshDescription, int32 Width, int32 Height,
-											 float GutterSpace, float Stretch, int32 ChartNum,
-											 bool bUseNormals, bool bRecomputeTangentSpace, bool bPrintDebugMessages) const = 0;
+	
+	virtual bool ParameterizeMeshDescription(FMeshDescription& MeshDescription, 
+		                                     int32 Width, 
+		                                     int32 Height,
+											 float GutterSpace, 
+		                                     float Stretch, 
+		                                     int32 ChartNum,
+											 bool bUseNormals, 
+		                                     bool bRecomputeTangentSpace, 
+		                                     bool bPrintDebugMessages) const = 0;
+
+	/**
+* Lower-level entry point:
+* Method that generates new UVs inconstant with the mesh defined by the VertexBuffer/IndexBuffer
+* according to the parameters specified in the FTextureAtlasDesc
+* The underlying code uses Isometeric approach (Iso-Charts) in UV generation.
+* This assumes that the caller has already generated the mesh adjacency information.
+*
+*
+*
+* @param InTextureAtlasDesc   Description of the texel resolution of the desired texture atlas.
+* @param VertexBuffer         Positions of vertices
+* @param IndexBuffer          Triangle definitions
+* @param AdjacencyBuffer      3 entries for each face, each entry is the face index of the face that is adjacent to the given face. -1 means no face.
+* @param StatusCallBack       Allows for termination of the UV generation.
+* @param UVVertexBuffer       Resulting UVs
+* @param UVIndexBuffer        Connectivity for the UVs
+* @param VertexRemapArray     Array that remaps the vertices split by the UV generation to the source vertices.
+* @param MaxStretch           On Input, The maximum amount of stretch between (0 - none and 1 - any), on Output the actual max stretch used.
+* @param MaxChartNumber       On Input Maximum number of charts required for the atlas.  If this is 0, will be parameterized solely on stretch.
+*                             Note, not a hard limit - isochart will stop when a valid charting is found that is greater or equal to this number.
+*                             On Output: Number of charts actually generated
+*
+* @return 'true' if the UV generation succeeded,  'false' if it failed.
+*/
+	virtual bool GenerateUVs(int32                    Width,
+			                 int32                    Height,
+		                     float                    GutterSpace,
+			        		 const TArray<FVector>&   VertexBuffer,
+		                     const TArray<int32>&     IndexBuffer,
+		                     const TArray<int32>&     AdjacencyBuffer,
+		                     TFunction<bool(float)>&  Callback,
+		                     TArray<FVector2D>&       UVVertexBuffer,
+		                     TArray<int32>&           UVIndexBuffer,
+		                     TArray<int32>&           VertexRemapArray,
+		                     float&                   MaxStretch,
+		                     int32&                   NumCharts) const = 0;
 };
