Fix static anaysis warnings

#rnx

diff --git a/Engine/Plugins/Experimental/ProxyLODPlugin/Source/ProxyLOD/Private/ProxyLODMeshParameterization.cpp b/Engine/Plugins/Experimental/ProxyLODPlugin/Source/ProxyLOD/Private/ProxyLODMeshParameterization.cpp
--- a/Engine/Plugins/Experimental/ProxyLODPlugin/Source/ProxyLOD/Private/ProxyLODMeshParameterization.cpp
+++ b/Engine/Plugins/Experimental/ProxyLODPlugin/Source/ProxyLOD/Private/ProxyLODMeshParameterization.cpp
@@ -25,34 +25,37 @@ bool ProxyLOD::GenerateUVs(const FTextureAtlasDesc& TextureAtlasDesc,
 	                       int32&                   NumCharts)
 {
 	const int32 NumVerts   = VertexBuffer.Num();
-	const int32 NumIndices = IndexBuffer.Num();
 	const int32 NumFaces   = IndexBuffer.Num() / 3;
 
 	// Copy data into DirectX library format
-	DirectX::XMFLOAT3* PosArray = new DirectX::XMFLOAT3[NumVerts];
-	uint32* Indices             = new uint32[NumIndices];
-	uint32* AdjacencyArray      = new uint32[3 * NumFaces];
+
+	TArray<DirectX::XMFLOAT3> PosArray;
+	PosArray.Empty(NumVerts);
+	for (int32 i = 0; i < NumVerts; ++i)
 	{
-		for (int32 i = 0; i < NumVerts; ++i)
-		{
-			const FVector& Vertex = VertexBuffer[i];
-			PosArray[i] = DirectX::XMFLOAT3(Vertex.X, Vertex.Y, Vertex.Z);
-		}
-		for (int32 i = 0; i < NumIndices; ++i)
-		{
-			Indices[i] = (uint32)IndexBuffer[i];
-		}
-		for (int32 i = 0, I = NumFaces * 3; i <  I; ++i)
-		{
-			AdjacencyArray[i] = static_cast<uint32>(AdjacencyBuffer[i]);
-		}
+		const FVector& Vertex = VertexBuffer[i];
+		PosArray.Emplace(DirectX::XMFLOAT3(Vertex.X, Vertex.Y, Vertex.Z));
+	}
+	
+	TArray<uint32> Indices;
+	Indices.Empty(3 * NumFaces);
+	for (int32 i = 0, I = NumFaces * 3; i < I; ++i)
+	{
+		Indices.Add(IndexBuffer[i]);
+	}
 
+	TArray<uint32> AdjacencyArray;
+	AdjacencyArray.Empty(3 * NumFaces);
+	for (int32 i = 0, I = NumFaces * 3; i < I; ++i)
+	{
+		AdjacencyArray.Add( static_cast<uint32>(AdjacencyBuffer[i]) );
 	}
 
 
+
 	// Verify the mesh is valid
 	{
-		HRESULT ValidateHR = DirectX::Validate(Indices, NumFaces, NumVerts, AdjacencyArray, DirectX::VALIDATE_BOWTIES, NULL);
+		HRESULT ValidateHR = DirectX::Validate(Indices.GetData(), NumFaces, NumVerts, AdjacencyArray.GetData(), DirectX::VALIDATE_BOWTIES, NULL);
 		if (FAILED(ValidateHR))
 		{
 			return false;
@@ -101,11 +104,11 @@ bool ProxyLOD::GenerateUVs(const FTextureAtlasDesc& TextureAtlasDesc,
 	std::vector<uint32> FacePartitioning;
 
 	// Generate UVs
-	HRESULT hr = DirectX::UVAtlasCreate(PosArray, NumVerts,
-		                                Indices, DXGI_FORMAT_R32_UINT, NumFaces,
+	HRESULT hr = DirectX::UVAtlasCreate(PosArray.GetData(), NumVerts,
+		                                Indices.GetData(), DXGI_FORMAT_R32_UINT, NumFaces,
 		                                NumChartsIn, MaxStretchIn,
 		                                Width, Height, Gutter,
-		                                AdjacencyArray, NULL /*false adj*/, pIMTArray /*IMTArray*/,
+		                                AdjacencyArray.GetData(), NULL /*false adj*/, pIMTArray /*IMTArray*/,
 		                                StatusCallBack, DirectX::UVATLAS_DEFAULT_CALLBACK_FREQUENCY,
 		                                DirectX::UVATLAS_DEFAULT, VB, IB,
 		                                &FacePartitioning, &RemapArray, &MaxStretchOut, &NumChartsOut);
@@ -151,12 +154,6 @@ bool ProxyLOD::GenerateUVs(const FTextureAtlasDesc& TextureAtlasDesc,
 		MaxStretch = MaxStretchOut;
 		NumCharts  = (int32)NumChartsOut;
 	}
-
-	// Clean up
-	delete[] pIMTArray;
-	delete[] PosArray;
-	delete[] Indices;
-	delete[] AdjacencyArray;
 	
 
 	return (hr == S_OK) ? true : false;
