Fix static analysis warning
#rnx

diff --git a/Engine/Plugins/Experimental/ProxyLODPlugin/Source/ProxyLOD/Private/ProxyLODMeshParameterization.cpp b/Engine/Plugins/Experimental/ProxyLODPlugin/Source/ProxyLOD/Private/ProxyLODMeshParameterization.cpp
--- a/Engine/Plugins/Experimental/ProxyLODPlugin/Source/ProxyLOD/Private/ProxyLODMeshParameterization.cpp
+++ b/Engine/Plugins/Experimental/ProxyLODPlugin/Source/ProxyLOD/Private/ProxyLODMeshParameterization.cpp
@@ -37,18 +37,18 @@ bool ProxyLOD::GenerateUVs(const FTextureAtlasDesc& TextureAtlasDesc,
 		PosArray.Emplace(DirectX::XMFLOAT3(Vertex.X, Vertex.Y, Vertex.Z));
 	}
 	
-	TArray<uint32> Indices;
+	TArray<uint32_t> Indices;
 	Indices.Empty(3 * NumFaces);
 	for (int32 i = 0, I = NumFaces * 3; i < I; ++i)
 	{
 		Indices.Add(IndexBuffer[i]);
 	}
 
-	TArray<uint32> AdjacencyArray;
+	TArray<uint32_t> AdjacencyArray;
 	AdjacencyArray.Empty(3 * NumFaces);
 	for (int32 i = 0, I = NumFaces * 3; i < I; ++i)
 	{
-		AdjacencyArray.Add( static_cast<uint32>(AdjacencyBuffer[i]) );
+		AdjacencyArray.Add( static_cast<uint32_t>(AdjacencyBuffer[i]) );
 	}
 
 
