MeshDescription now uses triangles as its primary face representation. They can and should be used in preference to polygons when a mesh is known to be entirely triangulated, as they are quicker, and can be used in exactly the same way.
Creating a triangle will also create a 'shadow' polygon for backward compatibility, although this is omitted when serializing.
As with polygons, any missing edges are created when creating a new triangle, and operations such as adjacency are available to triangles, just as for polygons.
The native serialization format for MeshDescription is now based around triangles.
Polygons of arbitrary number of edges may still be represented, and when created, will automatically generate the appropriate number of triangles.  Polygons of more than three sides automatically generate 'internal edges', which are owned by the polygon, and which will be regenerated if the polygon is retriangulated. Internal edges are otherwise not distinct from other edges, and can be used to determine connected triangles or vertices.
Direct access to public members of mesh element classes (FMeshVertex, etc) is now deprecated. In particular, vertex instances and edges now hold a list of connected triangles, and the ConnectedPolygons member is deprecated and will be removed (as this data can be derived from the ConnectedTriangles list).
Improved the MeshDescription API, providing more operations, and allowing arrays with different allocators to be provided as results containers.
Fixed several MeshEditing operations.
Added MeshDescription unit tests.

diff --git a/Engine/Plugins/Editor/MeshEditor/Source/MeshEditor/MeshEditorMode.cpp b/Engine/Plugins/Editor/MeshEditor/Source/MeshEditor/MeshEditorMode.cpp
--- a/Engine/Plugins/Editor/MeshEditor/Source/MeshEditor/MeshEditorMode.cpp
+++ b/Engine/Plugins/Editor/MeshEditor/Source/MeshEditor/MeshEditorMode.cpp
@@ -1436,7 +1436,7 @@ void FMeshEditorMode::UpdateDebugNormals()
 			// @todo mesheditor: total debug feature for now. Need a way of making this look nice.
 			const float Length = 10.0f; // @todo mesheditor: determine length of debug line from distance from the mesh origin to the camera?
 
-			for( const FVertexInstanceID VertexInstanceID : MeshDescription->GetPolygonPerimeterVertexInstances( PolygonID ) )
+			for( const FVertexInstanceID VertexInstanceID : MeshDescription->GetPolygonVertexInstances( PolygonID ) )
 			{
 				const FVector Position = VertexPositions[ MeshDescription->GetVertexInstanceVertex( VertexInstanceID ) ];
 				const FVector Normal = VertexNormals[ VertexInstanceID ];
@@ -2109,7 +2109,7 @@ void FMeshEditorMode::FrameSelectedElements( FEditorViewportClient* ViewportClie
 					{
 						const FPolygonID PolygonID( PolygonElement.ElementAddress.ElementID );
 
-						for( const FVertexInstanceID VertexInstanceID : MeshDescription->GetPolygonPerimeterVertexInstances( PolygonID ) )
+						for( const FVertexInstanceID VertexInstanceID : MeshDescription->GetPolygonVertexInstances( PolygonID ) )
 						{
 							const FVector VertexPosition = VertexPositions[ MeshDescription->GetVertexInstanceVertex( VertexInstanceID ) ];
 							BoundingBox += Component->GetComponentTransform().TransformPosition( VertexPosition );
@@ -2489,14 +2489,15 @@ void FMeshEditorMode::AddMeshElementToOverlay( UOverlayComponent* OverlayCompone
 					case EEditableMeshElementType::Polygon:
 					{
 						const FPolygonID PolygonID( MeshElement.ElementAddress.ElementID );
-						const int32 PolygonTriangleCount = EditableMesh->GetPolygonTriangulatedTriangleCount( PolygonID );
+						const TArray<FTriangleID>& TriangleIDs = MeshDescription->GetPolygonTriangleIDs( PolygonID );
 
-						for( int32 PolygonTriangle = 0; PolygonTriangle < PolygonTriangleCount; PolygonTriangle++ )
+						for( int32 PolygonTriangle = 0; PolygonTriangle < TriangleIDs.Num(); PolygonTriangle++ )
 						{
+							const FTriangleID TriangleID = TriangleIDs[ PolygonTriangle ];
 							FVector TriangleVertexPositions[ 3 ];
 							for( int32 TriangleVertex = 0; TriangleVertex < 3; TriangleVertex++ )
 							{
-								const FVertexInstanceID VertexInstanceID = EditableMesh->GetPolygonTriangulatedTriangle( PolygonID, PolygonTriangle ).GetVertexInstanceID( TriangleVertex );
+								const FVertexInstanceID VertexInstanceID = MeshDescription->GetTriangleVertexInstance( TriangleID, TriangleVertex );
 								const FVertexID VertexID = EditableMesh->GetVertexInstanceVertex( VertexInstanceID );
 								TriangleVertexPositions[ TriangleVertex ] = ComponentToWorldMatrix.TransformPosition( VertexPositions[ VertexID ] );
 							}
diff --git a/Engine/Plugins/Editor/MeshEditor/Source/PolygonModeling/SplitMesh.cpp b/Engine/Plugins/Editor/MeshEditor/Source/PolygonModeling/SplitMesh.cpp
--- a/Engine/Plugins/Editor/MeshEditor/Source/PolygonModeling/SplitMesh.cpp
+++ b/Engine/Plugins/Editor/MeshEditor/Source/PolygonModeling/SplitMesh.cpp
@@ -273,7 +273,7 @@ void USplitMeshCommand::Execute(IMeshEditorModeEditingContract& MeshEditorMode)
 		TSet<FPolygonGroupID> PolygonGroupSet;
 		for(const FPolygonID PolygonId : NewPolygonIds)
 		{
-			const TArray<FVertexInstanceID>& VertexInstanceIds = Mesh->GetMeshDescription()->GetPolygonPerimeterVertexInstances(PolygonId);
+			const TArray<FVertexInstanceID>& VertexInstanceIds = Mesh->GetMeshDescription()->GetPolygonVertexInstances(PolygonId);
 			// @todo (mlentine): Remove this when we don't need duplicate vertex instance ids.
 			TArray<FVertexInstanceID> NewVertexInstanceIds;
 			for(const FVertexInstanceID& VertexInstanceId : VertexInstanceIds)
@@ -297,7 +297,7 @@ void USplitMeshCommand::Execute(IMeshEditorModeEditingContract& MeshEditorMode)
 				CopyAllAttributesToDifferentElement(NewMeshDescription->VertexInstanceAttributes(), Mesh->GetMeshDescription()->VertexInstanceAttributes(), NewVertexInstanceID, VertexInstanceId);
 			}
 			TArray<FEdgeID> EdgeIds;
-			Mesh->GetMeshDescription()->GetPolygonEdges(PolygonId, EdgeIds);
+			Mesh->GetMeshDescription()->GetPolygonPerimeterEdges(PolygonId, EdgeIds);
 			for(const FEdgeID EdgeId : EdgeIds)
 			{
 				if(!EdgeSet.Contains(EdgeId))
diff --git a/Engine/Plugins/Editor/MeshEditor/Source/PolygonModeling/SplitPolygon.cpp b/Engine/Plugins/Editor/MeshEditor/Source/PolygonModeling/SplitPolygon.cpp
--- a/Engine/Plugins/Editor/MeshEditor/Source/PolygonModeling/SplitPolygon.cpp
+++ b/Engine/Plugins/Editor/MeshEditor/Source/PolygonModeling/SplitPolygon.cpp
@@ -162,9 +162,9 @@ void USplitPolygonCommand::ApplyDuringDrag( IMeshEditorModeEditingContract& Mesh
 
 					// @todo mesheditor splitpolygon: Ideally this would be more "fuzzy", and allow the interactor to extend beyond the range of the polygon.  But it will make figuring out which polygon to split more tricky.
 					// @todo mesheditor urgent: Can crash with "Colinear points in FMath::ComputeBaryCentric2D()"  Needs repro.
-					FMeshTriangle Triangle;
 					FVector TriangleVertexWeights;
-					if( EditableMesh->ComputeBarycentricWeightForPointOnPolygon( CandidatePolygonID, ComponentToWorld.InverseTransformPosition( LaserImpactOnPolygonPlane ), /* Out */ Triangle, /* Out */ TriangleVertexWeights ) )
+					FTriangleID TriangleID = EditableMesh->ComputeBarycentricWeightForPointOnPolygon( CandidatePolygonID, ComponentToWorld.InverseTransformPosition( LaserImpactOnPolygonPlane ), /* Out */ TriangleVertexWeights );
+					if( TriangleID != FTriangleID::Invalid )
 					{
 						const FVector SplitDirection = ( LaserImpactOnPolygonPlane - SplitStartLocation ).GetSafeNormal();
 
diff --git a/Engine/Plugins/Editor/MeshEditor/Source/PolygonModeling/UnifyNormals.cpp b/Engine/Plugins/Editor/MeshEditor/Source/PolygonModeling/UnifyNormals.cpp
--- a/Engine/Plugins/Editor/MeshEditor/Source/PolygonModeling/UnifyNormals.cpp
+++ b/Engine/Plugins/Editor/MeshEditor/Source/PolygonModeling/UnifyNormals.cpp
@@ -16,7 +16,7 @@ namespace UnifyNormalsCommandUtils
 	 * Returns -1 if the order is reversed
 	 * Returns 0 if the polygon doesn't contain the MeshEdge's vertices
 	 */
-	int8 GetPolygonEdgeDirection(const FMeshDescription* MeshDescription, const FMeshEdge& Edge, const FPolygonID& PolygonID)
+	int8 GetPolygonEdgeDirection(const FMeshDescription* MeshDescription, const FEdgeID EdgeID, const FPolygonID PolygonID)
 	{
 		if (!MeshDescription)
 		{
@@ -26,20 +26,21 @@ namespace UnifyNormalsCommandUtils
 		int8 EdgeDirection = 0;
 
 		// Iterate through the polygon vertices to find vertex 0 of the MeshEdge
-		const TArray<FVertexInstanceID>& VertexInstances = MeshDescription->GetPolygonPerimeterVertexInstances(PolygonID);
+		const FVertexID EdgeVertex0 = MeshDescription->GetEdgeVertex(EdgeID, 0);
+		const FVertexID EdgeVertex1 = MeshDescription->GetEdgeVertex(EdgeID, 1);
+		const TArray<FVertexInstanceID>& VertexInstances = MeshDescription->GetPolygonVertexInstances(PolygonID);
 		int32 NumVertices = VertexInstances.Num();
 		for (int32 Index = 0; Index < NumVertices; ++Index)
 		{
-			const FMeshVertexInstance& VertexInstance = MeshDescription->GetVertexInstance(VertexInstances[Index]);
-			if (VertexInstance.VertexID == Edge.VertexIDs[0])
+			if (MeshDescription->GetVertexInstanceVertex(VertexInstances[Index]) == EdgeVertex0)
 			{
 				// Vertex 1 is either the next vertex of the triangle => same direction
-				if (MeshDescription->GetVertexInstance(VertexInstances[(Index + 1) % NumVertices]).VertexID == Edge.VertexIDs[1])
+				if (MeshDescription->GetVertexInstanceVertex(VertexInstances[(Index + 1) % NumVertices]) == EdgeVertex1)
 				{
 					EdgeDirection = 1;
 				}
 				// Or the previous vertex of the triangle => opposite direction
-				else if (MeshDescription->GetVertexInstance(VertexInstances[(Index + NumVertices - 1) % NumVertices]).VertexID == Edge.VertexIDs[1])
+				else if (MeshDescription->GetVertexInstanceVertex(VertexInstances[(Index + NumVertices - 1) % NumVertices]) == EdgeVertex1)
 				{
 					EdgeDirection = -1;
 				}
@@ -109,20 +110,19 @@ void UUnifyNormalsCommand::Execute(IMeshEditorModeEditingContract& MeshEditorMod
 					CheckedPolygons.Add(PolygonID);
 
 					TArray<FEdgeID> PolygonEdges;
-					MeshDescription->GetPolygonEdges(PolygonID, PolygonEdges);
-					for (const FEdgeID& EdgeID : PolygonEdges)
+					MeshDescription->GetPolygonPerimeterEdges(PolygonID, PolygonEdges);
+					for (const FEdgeID EdgeID : PolygonEdges)
 					{
 						// For each edge of the polygon, check its direction and compare it with the edge direction of its neighbor polygon
-						const FMeshEdge& Edge = MeshDescription->GetEdge(EdgeID);
-						int8 PolygonEdgeDirection = UnifyNormalsCommandUtils::GetPolygonEdgeDirection(MeshDescription, Edge, PolygonID);
+						int8 PolygonEdgeDirection = UnifyNormalsCommandUtils::GetPolygonEdgeDirection(MeshDescription, EdgeID, PolygonID);
 
-						for (const FPolygonID& NeighborPolygonID : Edge.ConnectedPolygons)
+						for (const FPolygonID& NeighborPolygonID : MeshDescription->GetEdgeConnectedPolygons(EdgeID))
 						{
 							if (!CheckedPolygons.Contains(NeighborPolygonID))
 							{
 								CheckedPolygons.Add(NeighborPolygonID);
 
-								int8 ConnectedPolygonEdgeDirection = UnifyNormalsCommandUtils::GetPolygonEdgeDirection(MeshDescription, Edge, NeighborPolygonID);
+								int8 ConnectedPolygonEdgeDirection = UnifyNormalsCommandUtils::GetPolygonEdgeDirection(MeshDescription, EdgeID, NeighborPolygonID);
 
 								// For the initial polygon, its neighbor is considered flipped if the shared edge is in the same direction for both polygons
 								// For non-initial polygon, ie. a polygon that's already determined to be flipped, a neighbor polygon is flipped (with respect to the initial polygon)
@@ -137,7 +137,7 @@ void UUnifyNormalsCommand::Execute(IMeshEditorModeEditingContract& MeshEditorMod
 								{
 									// Add the vertices of the polygon at the boundary
 									Vertices.Reset();
-									MeshDescription->GetPolygonPerimeterVertices(NeighborPolygonID, Vertices);
+									MeshDescription->GetPolygonVertices(NeighborPolygonID, Vertices);
 									BoundaryPolygonVertices.Append(Vertices);
 								}
 							}
@@ -159,7 +159,7 @@ void UUnifyNormalsCommand::Execute(IMeshEditorModeEditingContract& MeshEditorMod
 			for (const FPolygonID& PolygonID : FlippedPolygons)
 			{
 				Vertices.Reset();
-				MeshDescription->GetPolygonPerimeterVertices(PolygonID, Vertices);
+				MeshDescription->GetPolygonVertices(PolygonID, Vertices);
 				FlippedPolygonVertices.Append(Vertices);
 			}
 
diff --git a/Engine/Plugins/Editor/SpeedTreeImporter/Source/SpeedTreeImporter/Private/SpeedTreeImportFactory.cpp b/Engine/Plugins/Editor/SpeedTreeImporter/Source/SpeedTreeImporter/Private/SpeedTreeImportFactory.cpp
--- a/Engine/Plugins/Editor/SpeedTreeImporter/Source/SpeedTreeImporter/Private/SpeedTreeImportFactory.cpp
+++ b/Engine/Plugins/Editor/SpeedTreeImporter/Source/SpeedTreeImporter/Private/SpeedTreeImportFactory.cpp
@@ -1796,10 +1796,7 @@ UObject* USpeedTreeImportFactory::FactoryCreateBinary7(UClass* InClass, UObject*
 								}
 
 								// Insert a polygon into the mesh
-								const FPolygonID NewPolygonID = MeshDescription->CreatePolygon(CurrentPolygonGroupID, CornerVertexInstanceIDs);
-								//Triangulate the polygon
-								FMeshPolygon& Polygon = MeshDescription->GetPolygon(NewPolygonID);
-								MeshDescription->ComputePolygonTriangulation(NewPolygonID, Polygon.Triangles);
+								MeshDescription->CreatePolygon(CurrentPolygonGroupID, CornerVertexInstanceIDs);
 							}
 						}
 
@@ -1945,10 +1942,7 @@ UObject* USpeedTreeImportFactory::FactoryCreateBinary7(UClass* InClass, UObject*
 							}
 
 							// Insert a polygon into the mesh
-							const FPolygonID NewPolygonID = MeshDescription->CreatePolygon(CurrentPolygonGroupID, CornerVertexInstanceIDs);
-							//Triangulate the polygon
-							FMeshPolygon& Polygon = MeshDescription->GetPolygon(NewPolygonID);
-							MeshDescription->ComputePolygonTriangulation(NewPolygonID, Polygon.Triangles);
+							MeshDescription->CreatePolygon(CurrentPolygonGroupID, CornerVertexInstanceIDs);
 						}
 					}
 
@@ -2329,10 +2323,7 @@ UObject* USpeedTreeImportFactory::FactoryCreateBinary8(UClass* InClass, UObject*
 					}
 
 					// Insert a polygon into the mesh
-					const FPolygonID NewPolygonID = MeshDescription->CreatePolygon(CurrentPolygonGroupID, CornerVertexInstanceIDs);
-					//Triangulate the polygon
-					FMeshPolygon& Polygon = MeshDescription->GetPolygon(NewPolygonID);
-					MeshDescription->ComputePolygonTriangulation(NewPolygonID, Polygon.Triangles);
+					MeshDescription->CreatePolygon(CurrentPolygonGroupID, CornerVertexInstanceIDs);
 				}
 			}
 
diff --git a/Engine/Plugins/Experimental/AlembicImporter/Source/AlembicLibrary/Private/AbcImporter.cpp b/Engine/Plugins/Experimental/AlembicImporter/Source/AlembicLibrary/Private/AbcImporter.cpp
--- a/Engine/Plugins/Experimental/AlembicImporter/Source/AlembicLibrary/Private/AbcImporter.cpp
+++ b/Engine/Plugins/Experimental/AlembicImporter/Source/AlembicLibrary/Private/AbcImporter.cpp
@@ -1348,10 +1348,7 @@ void FAbcImporter::GenerateMeshDescriptionFromSample(const FAbcMeshSample* Sampl
 
 		const FPolygonGroupID PolygonGroupID(Sample->MaterialIndices[TriangleIndex]);
 		// Insert a polygon into the mesh
-		const FPolygonID NewPolygonID = MeshDescription->CreatePolygon(PolygonGroupID, CornerVertexInstanceIDs);
-		//Triangulate the polygon
-		FMeshPolygon& Polygon = MeshDescription->GetPolygon(NewPolygonID);
-		MeshDescription->ComputePolygonTriangulation(NewPolygonID, Polygon.Triangles);
+		MeshDescription->CreatePolygon(PolygonGroupID, CornerVertexInstanceIDs);
 	}
 	//Set the edge hardness from the smooth group
 	FMeshDescriptionOperations::ConvertSmoothGroupToHardEdges(Sample->SmoothingGroupIndices, *MeshDescription);
diff --git a/Engine/Plugins/Experimental/BlastPlugin/Source/BlastAuthoring/Private/MeshUtility.cpp b/Engine/Plugins/Experimental/BlastPlugin/Source/BlastAuthoring/Private/MeshUtility.cpp
--- a/Engine/Plugins/Experimental/BlastPlugin/Source/BlastAuthoring/Private/MeshUtility.cpp
+++ b/Engine/Plugins/Experimental/BlastPlugin/Source/BlastAuthoring/Private/MeshUtility.cpp
@@ -51,12 +51,12 @@ void FMeshUtility::EditableMeshToBlastMesh(const UEditableMesh* SourceMesh, Nv::
 
 	for(const FPolygonID PolygonID : MeshDescription->Polygons().GetElementIDs())
 	{
-		const TArray<FMeshTriangle>& MeshTriangles = MeshDescription->GetPolygonTriangles(PolygonID);
-		for(const FMeshTriangle& MeshTriangle : MeshTriangles)
+		const TArray<FTriangleID>& TriangleIDs = MeshDescription->GetPolygonTriangleIDs(PolygonID);
+		for(const FTriangleID TriangleID : TriangleIDs)
 		{
 			for(int32 TriVertIndex = 0; TriVertIndex < 3; ++TriVertIndex)
 			{
-				const FVertexInstanceID VertexInstanceId = MeshTriangle.GetVertexInstanceID(TriVertIndex);
+				const FVertexInstanceID VertexInstanceId = MeshDescription->GetTriangleVertexInstance(TriangleID, TriVertIndex);
 				BlastIndices.Push(VertexInstanceId.GetValue());
 			}
 
@@ -74,7 +74,7 @@ void FMeshUtility::EditableMeshToBlastMesh(const UEditableMesh* SourceMesh, Nv::
 	const FPolygonArray& Polys = MeshDescription->Polygons();
 	for(int32 Facet = 0; Facet < FacetCount; Facet++)
 	{
-		FacetBuffer[Facet].materialId = Polys[FPolygonID(Facet)].PolygonGroupID.GetValue();
+		FacetBuffer[Facet].materialId = MeshDescription->GetPolygonPolygonGroup(FPolygonID(Facet)).GetValue();
 	}
 }
 
@@ -126,9 +126,9 @@ void FMeshUtility::EditableMeshToBlastMesh(const UEditableMesh* SourceMesh, int3
 	{
 		for(const FPolygonID PolygonID : PolygonGroupIDs)
 		{
-			const TArray<FMeshTriangle>& MeshTriangles = MeshDescription->GetPolygonTriangles(PolygonID);
+			const TArray<FTriangleID>& TriangleIDs = MeshDescription->GetPolygonTriangleIDs(PolygonID);
 
-			check(MeshTriangles.Num() == 1);
+			check(TriangleIDs.Num() == 1);
 
 #ifdef VALIDATE_INPUT
 			// does triangle match the geometry collection triangle
@@ -144,11 +144,11 @@ void FMeshUtility::EditableMeshToBlastMesh(const UEditableMesh* SourceMesh, int3
 				check(Visible[PolygonID.GetValue()]);
 #endif // VALIDATE_INPUT
 
-			for(const FMeshTriangle& MeshTriangle : MeshTriangles)
+			for(const FTriangleID TriangleID : TriangleIDs)
 			{
 				for(int32 TriVertIndex = 0; TriVertIndex < 3; ++TriVertIndex)
 				{
-					const FVertexInstanceID VertexInstanceID = MeshTriangle.GetVertexInstanceID(TriVertIndex);
+					const FVertexInstanceID VertexInstanceID = MeshDescription->GetTriangleVertexInstance(TriangleID, TriVertIndex);
 					BlastIndices.Push(VertexInstanceID.GetValue());
 				}
 				if(GeometryCollection)
diff --git a/Engine/Plugins/Experimental/GeometryProcessing/Source/MeshConversion/Private/MeshDescriptionBuilder.cpp b/Engine/Plugins/Experimental/GeometryProcessing/Source/MeshConversion/Private/MeshDescriptionBuilder.cpp
--- a/Engine/Plugins/Experimental/GeometryProcessing/Source/MeshConversion/Private/MeshDescriptionBuilder.cpp
+++ b/Engine/Plugins/Experimental/GeometryProcessing/Source/MeshConversion/Private/MeshDescriptionBuilder.cpp
@@ -88,7 +88,7 @@ FVector FMeshDescriptionBuilder::GetPosition(const FVertexID& VertexID)
 
 FVector FMeshDescriptionBuilder::GetPosition(const FVertexInstanceID& InstanceID)
 {
-	return VertexPositions.Get(MeshDescription->GetVertexInstance(InstanceID).VertexID, 0);
+	return VertexPositions.Get(MeshDescription->GetVertexInstanceVertex(InstanceID), 0);
 }
 
 
@@ -179,10 +179,6 @@ FPolygonID FMeshDescriptionBuilder::AppendPolygon(const TArray<FVertexID>& Verti
 
 	const FPolygonID NewPolygonID = MeshDescription->CreatePolygon(PolygonGroup, Polygon);
 
-	// compute triangulation
-	FMeshPolygon& NewPolygon = MeshDescription->GetPolygon(NewPolygonID);
-	MeshDescription->ComputePolygonTriangulation(NewPolygonID, NewPolygon.Triangles);
-
 	return NewPolygonID;
 }
 
@@ -198,10 +194,6 @@ FPolygonID FMeshDescriptionBuilder::AppendTriangle(const FVertexInstanceID& Inst
 
 	const FPolygonID NewPolygonID = MeshDescription->CreatePolygon(PolygonGroup, Polygon);
 
-	// compute triangulation
-	FMeshPolygon& NewPolygon = MeshDescription->GetPolygon(NewPolygonID);
-	MeshDescription->ComputePolygonTriangulation(NewPolygonID, NewPolygon.Triangles);
-
 	return NewPolygonID;
 }
 
@@ -332,23 +324,24 @@ void FMeshDescriptionBuilder::RecalculateInstanceNormals()
 	const FPolygonArray& Polygons = MeshDescription->Polygons();
 	for (const FPolygonID PolygonID : Polygons.GetElementIDs())
 	{
-		const TArray<FMeshTriangle>& Triangles = MeshDescription->GetPolygonTriangles(PolygonID);
-		for (FMeshTriangle Triangle : Triangles)
+		const TArray<FTriangleID>& TriangleIDs = MeshDescription->GetPolygonTriangleIDs(PolygonID);
+		for (const FTriangleID TriangleID : TriangleIDs)
 		{
-			FVector3d A = GetPosition(Triangle.VertexInstanceID0);
-			FVector3d B = GetPosition(Triangle.VertexInstanceID1);
-			FVector3d C = GetPosition(Triangle.VertexInstanceID2);
+			TArrayView<const FVertexInstanceID> TriangleVertexInstanceIDs = MeshDescription->GetTriangleVertexInstances(TriangleID);
+			FVector3d A = GetPosition(TriangleVertexInstanceIDs[0]);
+			FVector3d B = GetPosition(TriangleVertexInstanceIDs[1]);
+			FVector3d C = GetPosition(TriangleVertexInstanceIDs[2]);
 			double Area = 1.0;
 			FVector FaceNormal = VectorUtil::FastNormalArea(A, B, C, Area);
 			if (Area > FMathf::ZeroTolerance)
 			{
 				FaceNormal *= Area;
-				InstanceNormals.Set(Triangle.VertexInstanceID0,
-					InstanceNormals.Get(Triangle.VertexInstanceID0, 0) + FaceNormal);
-				InstanceNormals.Set(Triangle.VertexInstanceID1,
-					InstanceNormals.Get(Triangle.VertexInstanceID1, 0) + FaceNormal);
-				InstanceNormals.Set(Triangle.VertexInstanceID2,
-					InstanceNormals.Get(Triangle.VertexInstanceID2, 0) + FaceNormal);
+				InstanceNormals.Set(TriangleVertexInstanceIDs[0],
+					InstanceNormals.Get(TriangleVertexInstanceIDs[0], 0) + FaceNormal);
+				InstanceNormals.Set(TriangleVertexInstanceIDs[1],
+					InstanceNormals.Get(TriangleVertexInstanceIDs[1], 0) + FaceNormal);
+				InstanceNormals.Set(TriangleVertexInstanceIDs[2],
+					InstanceNormals.Get(TriangleVertexInstanceIDs[2], 0) + FaceNormal);
 			}
 		}
 	}
diff --git a/Engine/Plugins/Experimental/GeometryProcessing/Source/MeshConversion/Private/MeshDescriptionToDynamicMesh.cpp b/Engine/Plugins/Experimental/GeometryProcessing/Source/MeshConversion/Private/MeshDescriptionToDynamicMesh.cpp
--- a/Engine/Plugins/Experimental/GeometryProcessing/Source/MeshConversion/Private/MeshDescriptionToDynamicMesh.cpp
+++ b/Engine/Plugins/Experimental/GeometryProcessing/Source/MeshConversion/Private/MeshDescriptionToDynamicMesh.cpp
@@ -172,16 +172,13 @@ void FMeshDescriptionToDynamicMesh::Convert(const FMeshDescription* MeshIn, FDyn
 	{
 		FPolygonGroupID PolygonGroupID = MeshIn->GetPolygonPolygonGroup(PolygonID);
 
-		const TArray<FMeshTriangle>& Triangles = MeshIn->GetPolygonTriangles(PolygonID);
-		int NumTriangles = Triangles.Num();
+		const TArray<FTriangleID>& TriangleIDs = MeshIn->GetPolygonTriangleIDs(PolygonID);
+		int NumTriangles = TriangleIDs.Num();
 		for ( int TriIdx = 0; TriIdx < NumTriangles; ++TriIdx)
 		{
-			const FMeshTriangle& Triangle = Triangles[TriIdx];
+			const FTriangleID TriangleID = TriangleIDs[TriIdx];
 				
-			FVertexInstanceID InstanceTri[3];
-			InstanceTri[0] = Triangle.VertexInstanceID0; 
-			InstanceTri[1] = Triangle.VertexInstanceID1; 
-			InstanceTri[2] = Triangle.VertexInstanceID2;
+			TArrayView<const FVertexInstanceID> InstanceTri = MeshIn->GetTriangleVertexInstances(TriangleID);
 
 			int GroupID = 0;
 			if (GroupMode == EPrimaryGroupMode::SetToPolyGroup)
@@ -202,9 +199,9 @@ void FMeshDescriptionToDynamicMesh::Convert(const FMeshDescription* MeshIn, FDyn
 			
 
 			// append triangle
-			int VertexID0 = MeshIn->GetVertexInstance(Triangle.VertexInstanceID0).VertexID.GetValue();
-			int VertexID1 = MeshIn->GetVertexInstance(Triangle.VertexInstanceID1).VertexID.GetValue();
-			int VertexID2 = MeshIn->GetVertexInstance(Triangle.VertexInstanceID2).VertexID.GetValue();
+			int VertexID0 = MeshIn->GetVertexInstanceVertex(InstanceTri[0]).GetValue();
+			int VertexID1 = MeshIn->GetVertexInstanceVertex(InstanceTri[1]).GetValue();
+			int VertexID2 = MeshIn->GetVertexInstanceVertex(InstanceTri[2]).GetValue();
 			int NewTriangleID = MeshOut.AppendTriangle(VertexID0, VertexID1, VertexID2, GroupID);
 			
 			// if append failed due to non-manifold, duplicate verts
@@ -233,21 +230,21 @@ void FMeshDescriptionToDynamicMesh::Convert(const FMeshDescription* MeshIn, FDyn
 				}
 				if (bDuplicate[0])
 				{
-					FVertexID VertexID = MeshIn->GetVertexInstance(Triangle.VertexInstanceID0).VertexID;
+					FVertexID VertexID = MeshIn->GetVertexInstanceVertex(InstanceTri[0]);
 					const FVector Position = VertexPositions.Get(VertexID);
 					int NewVertIdx = MeshOut.AppendVertex(Position);
 					VertexID0 = NewVertIdx;
 				}
 				if (bDuplicate[1])
 				{
-					FVertexID VertexID = MeshIn->GetVertexInstance(Triangle.VertexInstanceID1).VertexID;
+					FVertexID VertexID = MeshIn->GetVertexInstanceVertex(InstanceTri[1]);
 					const FVector Position = VertexPositions.Get(VertexID);
 					int NewVertIdx = MeshOut.AppendVertex(Position);
 					VertexID1 = NewVertIdx;
 				}
 				if (bDuplicate[2])
 				{
-					FVertexID VertexID = MeshIn->GetVertexInstance(Triangle.VertexInstanceID2).VertexID;
+					FVertexID VertexID = MeshIn->GetVertexInstanceVertex(InstanceTri[2]);
 					const FVector Position = VertexPositions.Get(VertexID);
 					int NewVertIdx = MeshOut.AppendVertex(Position);
 					VertexID2 = NewVertIdx;
@@ -323,12 +320,9 @@ void FMeshDescriptionToDynamicMesh::CopyTangents(const FMeshDescription* SourceM
 	for (int TriID : TargetMesh->TriangleIndicesItr())
 	{
 		FIndex2i PolyTriIdx = TriToPolyTriMap[TriID];
-		const TArray<FMeshTriangle>& Triangles = SourceMesh->GetPolygonTriangles( FPolygonID(PolyTriIdx.A) );
-		const FMeshTriangle& Triangle = Triangles[PolyTriIdx.B];
-		FVertexInstanceID InstanceTri[3];
-		InstanceTri[0] = Triangle.VertexInstanceID0;
-		InstanceTri[1] = Triangle.VertexInstanceID1;
-		InstanceTri[2] = Triangle.VertexInstanceID2;
+		const TArray<FTriangleID>& TriangleIDs = SourceMesh->GetPolygonTriangleIDs( FPolygonID(PolyTriIdx.A) );
+		const FTriangleID TriangleID = TriangleIDs[PolyTriIdx.B];
+		TArrayView<const FVertexInstanceID> InstanceTri = SourceMesh->GetTriangleVertexInstances(TriangleID);
 		for (int j = 0; j < 3; ++j)
 		{
 			FVector Normal = InstanceNormals.Get(InstanceTri[j], 0);
diff --git a/Engine/Plugins/Experimental/ProxyLODPlugin/Source/ProxyLOD/Private/ProxyLODMeshAttrTransfer.cpp b/Engine/Plugins/Experimental/ProxyLODPlugin/Source/ProxyLOD/Private/ProxyLODMeshAttrTransfer.cpp
--- a/Engine/Plugins/Experimental/ProxyLODPlugin/Source/ProxyLOD/Private/ProxyLODMeshAttrTransfer.cpp
+++ b/Engine/Plugins/Experimental/ProxyLODPlugin/Source/ProxyLOD/Private/ProxyLODMeshAttrTransfer.cpp
@@ -76,23 +76,19 @@ void ProxyLOD::TransferMeshAttributes(const FClosestPolyField& SrcPolyField, FMe
 	{
 		FClosestPolyField::FPolyConstAccessor ConstPolyAccessor = SrcPolyField.GetPolyConstAccessor();
 
-		FVertexInstanceID WIdxs[3];
 		// loop over the faces
 		for (int32 CurrentRange = Range.begin(), EndRange = Range.end(); CurrentRange < EndRange; ++CurrentRange)
 		{
 			FPolygonID PolygonID = FPolygonID(CurrentRange);
-			const FMeshPolygon& Polygon = InOutMesh.GetPolygon(PolygonID);
 			int32 LastMaterialIndex = -1;
-			for (const FMeshTriangle& Triangle : Polygon.Triangles)
+			for (const FTriangleID TriangleID : InOutMesh.GetPolygonTriangleIDs(PolygonID))
 			{
 				// get the three corners for this Triangle
-				WIdxs[0] = Triangle.VertexInstanceID0;
-				WIdxs[1] = Triangle.VertexInstanceID1;
-				WIdxs[2] = Triangle.VertexInstanceID2;
+				TArrayView<const FVertexInstanceID> TriangleVerts = InOutMesh.GetTriangleVertexInstances(TriangleID);
 
 				for (int32 i = 0; i < 3; ++i)
 				{
-					const FVertexInstanceID& Idx = WIdxs[i];
+					const FVertexInstanceID Idx = TriangleVerts[i];
 					// world space location
 					const FVector& WSPos = VertexPositions[InOutMesh.GetVertexInstanceVertex(Idx)];
 
diff --git a/Engine/Plugins/Experimental/ProxyLODPlugin/Source/ProxyLOD/Private/ProxyLODMeshConvertUtils.cpp b/Engine/Plugins/Experimental/ProxyLODPlugin/Source/ProxyLOD/Private/ProxyLODMeshConvertUtils.cpp
--- a/Engine/Plugins/Experimental/ProxyLODPlugin/Source/ProxyLOD/Private/ProxyLODMeshConvertUtils.cpp
+++ b/Engine/Plugins/Experimental/ProxyLODPlugin/Source/ProxyLOD/Private/ProxyLODMeshConvertUtils.cpp
@@ -76,10 +76,7 @@ void ProxyLOD::MixedPolyMeshToRawMesh(const FMixedPolyMesh& SimpleMesh, FMeshDes
 		}
 
 		// Insert a polygon into the mesh
-		const FPolygonID NewPolygonID = DstRawMesh.CreatePolygon(PolygonGroupID, VertexInstanceIDs);
-		//Triangulate the polygon
-		FMeshPolygon& Polygon = DstRawMesh.GetPolygon(NewPolygonID);
-		DstRawMesh.ComputePolygonTriangulation(NewPolygonID, Polygon.Triangles);
+		DstRawMesh.CreatePolygon(PolygonGroupID, VertexInstanceIDs);
 	};
 
 	{
@@ -189,10 +186,7 @@ void ProxyLOD::AOSMeshToRawMesh(const FAOSMesh& AOSMesh, FMeshDescription& OutRa
 		}
 
 		// Insert a polygon into the mesh
-		const FPolygonID NewPolygonID = OutRawMesh.CreatePolygon(PolygonGroupID, VertexInstanceIDs);
-		//Triangulate the polygon
-		FMeshPolygon& Polygon = OutRawMesh.GetPolygon(NewPolygonID);
-		OutRawMesh.ComputePolygonTriangulation(NewPolygonID, Polygon.Triangles);
+		OutRawMesh.CreatePolygon(PolygonGroupID, VertexInstanceIDs);
 	};
 
 	{
@@ -326,10 +320,7 @@ void ProxyLOD::VertexDataMeshToRawMesh(const FVertexDataMesh& SrcVertexDataMesh,
 		}
 
 		// Insert a polygon into the mesh
-		const FPolygonID NewPolygonID = OutRawMesh.CreatePolygon(PolygonGroupID, VertexInstanceIDs);
-		//Triangulate the polygon
-		FMeshPolygon& Polygon = OutRawMesh.GetPolygon(NewPolygonID);
-		OutRawMesh.ComputePolygonTriangulation(NewPolygonID, Polygon.Triangles);
+		OutRawMesh.CreatePolygon(PolygonGroupID, VertexInstanceIDs);
 	};
 
 	{
@@ -378,12 +369,7 @@ void ProxyLOD::RawMeshToVertexDataMesh(const FMeshDescription& SrcRawMesh, FVert
 
 	const uint32 DstNumPositions = SrcRawMesh.Vertices().Num();
 
-	uint32 DstNumIndexes = 0;
-	for (const FPolygonID& PolygonID : SrcRawMesh.Polygons().GetElementIDs())
-	{
-		const FMeshPolygon& Polygon = SrcRawMesh.GetPolygon(PolygonID);
-		DstNumIndexes += Polygon.Triangles.Num() * 3;
-	}
+	uint32 DstNumIndexes = SrcRawMesh.Triangles().Num() * 3;
 
 	// Copy the vertices over
 	TMap<FVertexID, uint32> VertexIDToDstVertexIndex;
@@ -418,14 +404,13 @@ void ProxyLOD::RawMeshToVertexDataMesh(const FMeshDescription& SrcRawMesh, FVert
 	ResizeArray(DstUVs, DstNumPositions);
 
 	//Iterate all triangle and add the indices
-	for (const FPolygonID& PolygonID : SrcRawMesh.Polygons().GetElementIDs())
+	for (const FPolygonID PolygonID : SrcRawMesh.Polygons().GetElementIDs())
 	{
-		const FMeshPolygon& Polygon = SrcRawMesh.GetPolygon(PolygonID);
-		for (const FMeshTriangle& Triangle : Polygon.Triangles)
+		for (const FTriangleID TriangleID : SrcRawMesh.GetPolygonTriangleIDs(PolygonID))
 		{
 			for (int32 Corner = 0; Corner < 3; ++Corner)
 			{
-				const FVertexInstanceID& VertexInstanceID = Triangle.GetVertexInstanceID(Corner);
+				const FVertexInstanceID& VertexInstanceID = SrcRawMesh.GetTriangleVertexInstance(TriangleID, Corner);
 				DstIndices[VertexInstanceCount] = VertexIDToDstVertexIndex[SrcRawMesh.GetVertexInstanceVertex(VertexInstanceID)];
 
 				// Copy the tangent space:
diff --git a/Engine/Plugins/Experimental/ProxyLODPlugin/Source/ProxyLOD/Private/ProxyLODMeshParameterization.cpp b/Engine/Plugins/Experimental/ProxyLODPlugin/Source/ProxyLOD/Private/ProxyLODMeshParameterization.cpp
--- a/Engine/Plugins/Experimental/ProxyLODPlugin/Source/ProxyLOD/Private/ProxyLODMeshParameterization.cpp
+++ b/Engine/Plugins/Experimental/ProxyLODPlugin/Source/ProxyLOD/Private/ProxyLODMeshParameterization.cpp
@@ -446,11 +446,7 @@ void ProxyLOD::GenerateAdjacency(const FVertexDataMesh& Mesh, std::vector<uint32
 }
 void ProxyLOD::GenerateAdjacency(const FMeshDescription& RawMesh, std::vector<uint32>& AdjacencyArray)
 {
-	uint32 NumTris = 0;
-	for (const FPolygonID& PolygonID : RawMesh.Polygons().GetElementIDs())
-	{
-		NumTris += RawMesh.GetPolygonTriangles(PolygonID).Num();
-	}
+	uint32 NumTris = RawMesh.Triangles().Num();
 	const uint32 NumVerts = RawMesh.Vertices().Num();
 	const uint32 AdjacencySize = RawMesh.VertexInstances().Num(); // 3 for each face
 
diff --git a/Engine/Plugins/Experimental/ProxyLODPlugin/Source/ProxyLOD/Private/ProxyLODMeshTypes.cpp b/Engine/Plugins/Experimental/ProxyLODPlugin/Source/ProxyLOD/Private/ProxyLODMeshTypes.cpp
--- a/Engine/Plugins/Experimental/ProxyLODPlugin/Source/ProxyLOD/Private/ProxyLODMeshTypes.cpp
+++ b/Engine/Plugins/Experimental/ProxyLODPlugin/Source/ProxyLOD/Private/ProxyLODMeshTypes.cpp
@@ -20,22 +20,18 @@ FMeshDescriptionAdapter::FMeshDescriptionAdapter(const FMeshDescriptionAdapter&
 void FMeshDescriptionAdapter::InitializeCacheData()
 {
 	VertexPositions = RawMesh->VertexAttributes().GetAttributesRef<FVector>(MeshAttribute::Vertex::Position);
-	TriangleCount = 0;
-	for (const FPolygonID& PolygonID : RawMesh->Polygons().GetElementIDs())
-	{
-		TriangleCount += RawMesh->GetPolygonTriangles(PolygonID).Num();
-	}
+	TriangleCount = RawMesh->Triangles().Num();
 
 	IndexBuffer.Reserve(TriangleCount * 3);
 
-	for (const FPolygonID& PolygonID : RawMesh->Polygons().GetElementIDs())
+	for (const FPolygonID PolygonID : RawMesh->Polygons().GetElementIDs())
 	{
-		const auto& Triangles = RawMesh->GetPolygonTriangles(PolygonID);
-		for (const FMeshTriangle& Tri : Triangles)
+		const TArray<FTriangleID>& TriangleIDs = RawMesh->GetPolygonTriangleIDs(PolygonID);
+		for (const FTriangleID TriangleID : TriangleIDs)
 		{
-			IndexBuffer.Add(Tri.VertexInstanceID0);
-			IndexBuffer.Add(Tri.VertexInstanceID1);
-			IndexBuffer.Add(Tri.VertexInstanceID2);
+			IndexBuffer.Add(RawMesh->GetTriangleVertexInstance(TriangleID, 0));
+			IndexBuffer.Add(RawMesh->GetTriangleVertexInstance(TriangleID, 1));
+			IndexBuffer.Add(RawMesh->GetTriangleVertexInstance(TriangleID, 2));
 		}
 	}
 
@@ -82,24 +78,20 @@ FMeshDescriptionArrayAdapter::FMeshDescriptionArrayAdapter(const TArray<const FM
 
 		PointCount += size_t(RawMesh->Vertices().Num());
 		// Sum up all the polys in this mesh.
-		int32 MeshPolyCount = 0;
-		for (const FPolygonID& PolygonID : RawMesh->Polygons().GetElementIDs())
-		{
-			MeshPolyCount += RawMesh->GetPolygonTriangles(PolygonID).Num();
-		}
+		int32 MeshPolyCount = RawMesh->Triangles().Num();
 
 		// Construct local index buffer for this mesh
 		IndexBufferArray.push_back(std::vector<FVertexInstanceID>());
 		std::vector<FVertexInstanceID>& IndexBuffer = IndexBufferArray[MeshIdx];
 		IndexBuffer.reserve(MeshPolyCount * 3);
-		for (const FPolygonID& PolygonID : RawMesh->Polygons().GetElementIDs())
+		for (const FPolygonID PolygonID : RawMesh->Polygons().GetElementIDs())
 		{
-			const auto& Triangles = RawMesh->GetPolygonTriangles(PolygonID);
-			for (const FMeshTriangle& Tri : Triangles)
+			const TArray<FTriangleID>& TriangleIDs = RawMesh->GetPolygonTriangleIDs(PolygonID);
+			for (const FTriangleID TriangleID : TriangleIDs)
 			{
-				IndexBuffer.push_back(Tri.VertexInstanceID0);
-				IndexBuffer.push_back(Tri.VertexInstanceID1);
-				IndexBuffer.push_back(Tri.VertexInstanceID2);
+				IndexBuffer.push_back(RawMesh->GetTriangleVertexInstance(TriangleID, 0));
+				IndexBuffer.push_back(RawMesh->GetTriangleVertexInstance(TriangleID, 1));
+				IndexBuffer.push_back(RawMesh->GetTriangleVertexInstance(TriangleID, 2));
 			}
 		}
 
@@ -132,24 +124,20 @@ FMeshDescriptionArrayAdapter::FMeshDescriptionArrayAdapter(const TArray<FMeshMer
 		PointCount += size_t(RawMesh->Vertices().Num());
 
 		// Sum up all the polys in this mesh.
-		int32 MeshPolyCount = 0;
-		for (const FPolygonID& PolygonID : RawMesh->Polygons().GetElementIDs())
-		{
-			MeshPolyCount += RawMesh->GetPolygonTriangles(PolygonID).Num();
-		}
+		int32 MeshPolyCount = RawMesh->Triangles().Num();
 		PolyCount += MeshPolyCount;
 		// Construct local index buffer for this mesh
 		IndexBufferArray.push_back(std::vector<FVertexInstanceID>());
 		std::vector<FVertexInstanceID>& IndexBuffer = IndexBufferArray[MeshIdx];
 		IndexBuffer.reserve(MeshPolyCount * 3);
-		for (const FPolygonID& PolygonID : RawMesh->Polygons().GetElementIDs())
+		for (const FPolygonID PolygonID : RawMesh->Polygons().GetElementIDs())
 		{
-			const auto& Triangles = RawMesh->GetPolygonTriangles(PolygonID);
-			for (const FMeshTriangle& Tri : Triangles)
+			const TArray<FTriangleID>& TriangleIDs = RawMesh->GetPolygonTriangleIDs(PolygonID);
+			for (const FTriangleID TriangleID : TriangleIDs)
 			{
-				IndexBuffer.push_back(Tri.VertexInstanceID0);
-				IndexBuffer.push_back(Tri.VertexInstanceID1);
-				IndexBuffer.push_back(Tri.VertexInstanceID2);
+				IndexBuffer.push_back(RawMesh->GetTriangleVertexInstance(TriangleID, 0));
+				IndexBuffer.push_back(RawMesh->GetTriangleVertexInstance(TriangleID, 1));
+				IndexBuffer.push_back(RawMesh->GetTriangleVertexInstance(TriangleID, 2));
 			}
 		}
 
@@ -179,11 +167,7 @@ FMeshDescriptionArrayAdapter::FMeshDescriptionArrayAdapter(const TArray<FMeshMer
 		PointCount += size_t(RawMesh->Vertices().Num());
 		
 		// Sum up all the polys in this mesh.
-		int32 MeshPolyCount = 0;
-		for (const FPolygonID& PolygonID : RawMesh->Polygons().GetElementIDs())
-		{
-			MeshPolyCount += RawMesh->GetPolygonTriangles(PolygonID).Num();
-		}
+		int32 MeshPolyCount = RawMesh->Triangles().Num();
 
 		// Construct local index buffer for this mesh
 		IndexBufferArray.push_back(std::vector<FVertexInstanceID>());
@@ -191,12 +175,12 @@ FMeshDescriptionArrayAdapter::FMeshDescriptionArrayAdapter(const TArray<FMeshMer
 		IndexBuffer.reserve(MeshPolyCount * 3);
 		for (const FPolygonID& PolygonID : RawMesh->Polygons().GetElementIDs())
 		{
-			const auto& Triangles = RawMesh->GetPolygonTriangles(PolygonID);
-			for (const FMeshTriangle& Tri : Triangles)
+			const TArray<FTriangleID>& TriangleIDs = RawMesh->GetPolygonTriangleIDs(PolygonID);
+			for (const FTriangleID TriangleID : TriangleIDs)
 			{
-				IndexBuffer.push_back(Tri.VertexInstanceID0);
-				IndexBuffer.push_back(Tri.VertexInstanceID1);
-				IndexBuffer.push_back(Tri.VertexInstanceID2);
+				IndexBuffer.push_back(RawMesh->GetTriangleVertexInstance(TriangleID, 0));
+				IndexBuffer.push_back(RawMesh->GetTriangleVertexInstance(TriangleID, 1));
+				IndexBuffer.push_back(RawMesh->GetTriangleVertexInstance(TriangleID, 2));
 			}
 		}
 
@@ -320,12 +304,11 @@ FMeshDescriptionArrayAdapter::FRawPoly FMeshDescriptionArrayAdapter::GetRawPoly(
 	RawPoly.FaceMaterialIndex = RawMesh.GetPolygonPolygonGroup(PolygonID).GetValue();
 	RawPoly.FaceSmoothingMask = AttributesGetter->FaceSmoothingMasks[LocalFaceNumber];
 
-	for (const FMeshTriangle& Triangle : RawMesh.GetPolygonTriangles(PolygonID))
+	for (const FTriangleID TriangleID : RawMesh.GetPolygonTriangleIDs(PolygonID))
 	{
-		FVertexInstanceID VertexInstanceIDs[3];
+		TArrayView<const FVertexInstanceID> VertexInstanceIDs = RawMesh.GetTriangleVertexInstances(TriangleID);
 		for (int32 i = 0; i < 3; ++i)
 		{
-			VertexInstanceIDs[i] = Triangle.GetVertexInstanceID(i);
 			RawPoly.VertexPositions[i] = AttributesGetter->VertexPositions[RawMesh.GetVertexInstanceVertex(VertexInstanceIDs[i])];
 		}
 
diff --git a/Engine/Plugins/Experimental/ProxyLODPlugin/Source/ProxyLOD/Private/ProxyLODMeshTypes.h b/Engine/Plugins/Experimental/ProxyLODPlugin/Source/ProxyLOD/Private/ProxyLODMeshTypes.h
--- a/Engine/Plugins/Experimental/ProxyLODPlugin/Source/ProxyLOD/Private/ProxyLODMeshTypes.h
+++ b/Engine/Plugins/Experimental/ProxyLODPlugin/Source/ProxyLOD/Private/ProxyLODMeshTypes.h
@@ -323,11 +323,7 @@ public:
 			VertexInstanceColors = MeshDescription->VertexInstanceAttributes().GetAttributesRef<FVector4>(MeshAttribute::VertexInstance::Color);
 			VertexInstanceUVs = MeshDescription->VertexInstanceAttributes().GetAttributesRef<FVector2D>(MeshAttribute::VertexInstance::TextureCoordinate);
 
-			TriangleCount = 0;
-			for (FPolygonID PolygonID : MeshDescription->Polygons().GetElementIDs())
-			{
-				TriangleCount += MeshDescription->GetPolygonTriangles(PolygonID).Num();
-			}
+			TriangleCount = MeshDescription->Triangles().Num();
 			FaceSmoothingMasks.AddZeroed(TriangleCount);
 			FMeshDescriptionOperations::ConvertHardEdgesToSmoothGroup(*MeshDescription, FaceSmoothingMasks);
 		}
diff --git a/Engine/Plugins/Experimental/ProxyLODPlugin/Source/ProxyLOD/Private/ProxyLODMeshUtilities.cpp b/Engine/Plugins/Experimental/ProxyLODPlugin/Source/ProxyLOD/Private/ProxyLODMeshUtilities.cpp
--- a/Engine/Plugins/Experimental/ProxyLODPlugin/Source/ProxyLOD/Private/ProxyLODMeshUtilities.cpp
+++ b/Engine/Plugins/Experimental/ProxyLODPlugin/Source/ProxyLOD/Private/ProxyLODMeshUtilities.cpp
@@ -1480,11 +1480,7 @@ int32 CorrectCollapsedWalls(const ProxyLOD::FkDOPTree& kDOPTree,
 	ProxyLOD::FUnitTransformDataProvider kDOPDataProvider(kDOPTree);
 
 	// Number of triangle in our mesh
-	int32 NumTriangle = 0;
-	for (const FPolygonID& PolygonID : MeshDescription.Polygons().GetElementIDs())
-	{
-		NumTriangle += MeshDescription.GetPolygonTriangles(PolygonID).Num();
-	}
+	int32 NumTriangle = MeshDescription.Triangles().Num();
 
 	// This will hold the intersecting faces for each edge.
 
@@ -1686,14 +1682,13 @@ void ProxyLOD::ColorPartitions(FMeshDescription& InOutRawMesh, const std::vector
 	int32 TriangleIndex = 0;
 	TMap<uint32, FVertexInstanceID> WedgeIndexToVertexInstanceID;
 	WedgeIndexToVertexInstanceID.Reserve(InOutRawMesh.VertexInstances().Num());
-	for (const FPolygonID& PolygonID : InOutRawMesh.Polygons().GetElementIDs())
+	for (const FPolygonID PolygonID : InOutRawMesh.Polygons().GetElementIDs())
 	{
-		const FMeshPolygon& Polygon = InOutRawMesh.GetPolygon(PolygonID);
-		for (const FMeshTriangle& Triangle : Polygon.Triangles)
+		for (const FTriangleID TriangleID : InOutRawMesh.GetPolygonTriangleIDs(PolygonID))
 		{
 			for (int32 Corner = 0; Corner < 3; ++Corner)
 			{
-				WedgeIndexToVertexInstanceID.Add((TriangleIndex * 3) + Corner, Triangle.GetVertexInstanceID(Corner));
+				WedgeIndexToVertexInstanceID.Add((TriangleIndex * 3) + Corner, InOutRawMesh.GetTriangleVertexInstance(TriangleID, Corner));
 			}
 			TriangleIndex++;
 		}
@@ -1743,15 +1738,13 @@ void ProxyLOD::AddWedgeColors(FMeshDescription& RawMesh)
 
 	//Recolor the vertex instances
 	int32 TriangleIndex = 0;
-	for (const FPolygonID& PolygonID : RawMesh.Polygons().GetElementIDs())
+	for (const FPolygonID PolygonID : RawMesh.Polygons().GetElementIDs())
 	{
-		const FMeshPolygon& Polygon = RawMesh.GetPolygon(PolygonID);
-		for (const FMeshTriangle& Triangle : Polygon.Triangles)
+		for (const FTriangleID TriangleID : RawMesh.GetPolygonTriangleIDs(PolygonID))
 		{
 			for (int32 Corner = 0; Corner < 3; ++Corner)
 			{
-
-				VertexInstanceColors[Triangle.GetVertexInstanceID(Corner)] = FLinearColor(ColorRange[((TriangleIndex*3) + Corner) % 13]);
+				VertexInstanceColors[RawMesh.GetTriangleVertexInstance(TriangleID, Corner)] = FLinearColor(ColorRange[((TriangleIndex*3) + Corner) % 13]);
 			}
 			TriangleIndex++;
 		}
diff --git a/Engine/Plugins/Experimental/ProxyLODPlugin/Source/ProxyLOD/Private/ProxyLODkDOPUtils.cpp b/Engine/Plugins/Experimental/ProxyLODPlugin/Source/ProxyLOD/Private/ProxyLODkDOPUtils.cpp
--- a/Engine/Plugins/Experimental/ProxyLODPlugin/Source/ProxyLOD/Private/ProxyLODkDOPUtils.cpp
+++ b/Engine/Plugins/Experimental/ProxyLODPlugin/Source/ProxyLOD/Private/ProxyLODkDOPUtils.cpp
@@ -42,11 +42,7 @@ void ProxyLOD::BuildkDOPTree(const FMeshDescription& MeshDescription, FkDOPTree&
 {
 	TVertexAttributesConstRef<FVector> VertexPositions = MeshDescription.VertexAttributes().GetAttributesRef<FVector>(MeshAttribute::Vertex::Position);
 
-	uint32 NumSrcPoly = 0;
-	for (const FPolygonID& PolygonID : MeshDescription.Polygons().GetElementIDs())
-	{
-		NumSrcPoly += MeshDescription.GetPolygonTriangles(PolygonID).Num();
-	}
+	uint32 NumSrcPoly = MeshDescription.Triangles().Num();
 
 	TArray<FkDOPBuildTriangle> BuildTriangleArray;
 
diff --git a/Engine/Plugins/Importers/USDImporter/Source/USDImporter/Private/StaticMeshImporter.cpp b/Engine/Plugins/Importers/USDImporter/Source/USDImporter/Private/StaticMeshImporter.cpp
--- a/Engine/Plugins/Importers/USDImporter/Source/USDImporter/Private/StaticMeshImporter.cpp
+++ b/Engine/Plugins/Importers/USDImporter/Source/USDImporter/Private/StaticMeshImporter.cpp
@@ -247,11 +247,6 @@ bool FUSDStaticMeshImportState::AddPolygons(FMeshDescriptionWrapper& DestMeshWra
 		{
 			MeshDescription->ReversePolygonFacing(NewPolygonID);
 		}
-		else
-		{
-			FMeshPolygon& Polygon = MeshDescription->GetPolygon(NewPolygonID);
-			MeshDescription->ComputePolygonTriangulation(NewPolygonID, Polygon.Triangles);
-		}
 	}
 
 	return true;
diff --git a/Engine/Plugins/Runtime/EditableMesh/Source/EditableMesh/EditableGeometryCollectionAdapter.cpp b/Engine/Plugins/Runtime/EditableMesh/Source/EditableMesh/EditableGeometryCollectionAdapter.cpp
--- a/Engine/Plugins/Runtime/EditableMesh/Source/EditableMesh/EditableGeometryCollectionAdapter.cpp
+++ b/Engine/Plugins/Runtime/EditableMesh/Source/EditableMesh/EditableGeometryCollectionAdapter.cpp
@@ -216,10 +216,6 @@ void UEditableGeometryCollectionAdapter::InitEditableGeometryCollection( UEditab
 							FAdaptorPolygon& NewRenderingPolygon = RenderingPolygons[NewPolygonID];
 							NewRenderingPolygon.PolygonGroupID = NewPolygonGroupID;
 							NewRenderingPolygon.TriangulatedPolygonTriangleIndices.Add(NewTriangleID);
-
-							// Add triangle to polygon triangulation array
-							MeshDescription->GetPolygonTriangles(NewPolygonID).Add(NewTriangle);
-
 						}
 						else
 						{
@@ -319,14 +315,18 @@ void UEditableGeometryCollectionAdapter::InitializeFromEditableMesh( const UEdit
 		FAdaptorPolygon& RenderingPolygon = RenderingPolygons[ PolygonID ];
 		RenderingPolygon.PolygonGroupID = PolygonGroupID;
 
-		const TArray<FMeshTriangle>& Triangles = MeshDescription->GetPolygonTriangles( PolygonID );
-		for( const FMeshTriangle& Triangle : Triangles )
+		const TArray<FTriangleID>& TriangleIDs = MeshDescription->GetPolygonTriangleIDs( PolygonID );
+		for( const FTriangleID TriangleID : TriangleIDs )
 		{
-			const FAdaptorTriangleID TriangleID = RenderingPolygonGroup.Triangles.Add( Triangle );
-			RenderingPolygon.TriangulatedPolygonTriangleIndices.Add( TriangleID );
+			FMeshTriangle Triangle;
+			Triangle.SetVertexInstanceID( 0, MeshDescription->GetTriangleVertexInstance( TriangleID, 0 ) );
+			Triangle.SetVertexInstanceID( 1, MeshDescription->GetTriangleVertexInstance( TriangleID, 1 ) );
+			Triangle.SetVertexInstanceID( 2, MeshDescription->GetTriangleVertexInstance( TriangleID, 2 ) );
+			const FAdaptorTriangleID AdapterTriangleID = RenderingPolygonGroup.Triangles.Add( Triangle );
+			RenderingPolygon.TriangulatedPolygonTriangleIndices.Add( AdapterTriangleID );
 		}
 
-		RenderingPolygonGroup.MaxTriangles += Triangles.Num();
+		RenderingPolygonGroup.MaxTriangles += TriangleIDs.Num();
 	}
 
 }
@@ -409,18 +409,22 @@ void UEditableGeometryCollectionAdapter::OnRebuildRenderMesh(const UEditableMesh
 				const FPolygonGroupID& PolygonGroupID = MeshDescription->GetPolygonPolygonGroup(PolygonID);
 				int32 PolygonGroupIDValue = PolygonGroupID.GetValue();
 				int32 PolygonIDValue = PolygonID.GetValue();
-				const TArray<FMeshTriangle>& Triangles = MeshDescription->GetPolygonTriangles(PolygonID);
-				for (const FMeshTriangle& MeshTriangle : Triangles)
+				const TArray<FTriangleID>& TriangleIDs = MeshDescription->GetPolygonTriangleIDs(PolygonID);
+				for (const FTriangleID TriangleID : TriangleIDs)
 				{
 					uint32 ElementIndex = Collection->AddElements(1, FGeometryCollection::FacesGroup);
 
+					const int32 V0 = MeshDescription->GetTriangleVertexInstance(TriangleID, 0).GetValue();
+					const int32 V1 = MeshDescription->GetTriangleVertexInstance(TriangleID, 1).GetValue();
+					const int32 V2 = MeshDescription->GetTriangleVertexInstance(TriangleID, 2).GetValue();
+						
 					// might need a lookup here VertexInstanceID to GCVector array index
-					GCIndices[ElementIndex] = FIntVector(MeshTriangle.VertexInstanceID0.GetValue(), MeshTriangle.VertexInstanceID1.GetValue(), MeshTriangle.VertexInstanceID2.GetValue());
+					GCIndices[ElementIndex] = FIntVector(V0, V1, V2);
 					GCVisible[ElementIndex] = true;
 
-					GCBoneMap[MeshTriangle.VertexInstanceID0.GetValue()] = PolygonGroupIDValue;
-					GCBoneMap[MeshTriangle.VertexInstanceID1.GetValue()] = PolygonGroupIDValue;
-					GCBoneMap[MeshTriangle.VertexInstanceID2.GetValue()] = PolygonGroupIDValue;
+					GCBoneMap[V0] = PolygonGroupIDValue;
+					GCBoneMap[V1] = PolygonGroupIDValue;
+					GCBoneMap[V2] = PolygonGroupIDValue;
 
 				}
 			}
@@ -883,9 +887,9 @@ void UEditableGeometryCollectionAdapter::DeletePolygonTriangles( const UEditable
 
 					// sanity check
 					FIntVector TriIndices = GCIndices[TriangleIndexToRemove.GetValue()];
-					check(TriToRemove.VertexInstanceID0.GetValue() == TriIndices[0]);
-					check(TriToRemove.VertexInstanceID1.GetValue() == TriIndices[1]);
-					check(TriToRemove.VertexInstanceID2.GetValue() == TriIndices[2]);
+					check(TriToRemove.GetVertexInstanceID(0).GetValue() == TriIndices[0]);
+					check(TriToRemove.GetVertexInstanceID(1).GetValue() == TriIndices[1]);
+					check(TriToRemove.GetVertexInstanceID(2).GetValue() == TriIndices[2]);
 
 					UE_LOG(LogGeometryCollectionAdapter, Log, TEXT("Deleting Tri_ID %d, Indices %d %d %d "), TriangleIndexToRemove.GetValue(), TriIndices[0], TriIndices[1], TriIndices[2]);
 
