ProxyLOD: When converting internal representations to FMeshDescription, make sure the dest mesh has been emptied first and initialized to hold all static mesh attrs.
#rnx

diff --git a/Engine/Plugins/Experimental/ProxyLODPlugin/Source/ProxyLOD/Private/ProxyLODMeshConvertUtils.cpp b/Engine/Plugins/Experimental/ProxyLODPlugin/Source/ProxyLOD/Private/ProxyLODMeshConvertUtils.cpp
--- a/Engine/Plugins/Experimental/ProxyLODPlugin/Source/ProxyLOD/Private/ProxyLODMeshConvertUtils.cpp
+++ b/Engine/Plugins/Experimental/ProxyLODPlugin/Source/ProxyLOD/Private/ProxyLODMeshConvertUtils.cpp
@@ -5,9 +5,19 @@
 
 #include "MeshDescriptionOperations.h"
 
+
+static void ResetStaticMeshDescription(FMeshDescription& MeshDecription)
+{
+	MeshDecription.Empty();
+	FStaticMeshAttributes Attributes(MeshDecription);
+	Attributes.Register();
+}
+
 // Convert QuadMesh to Triangles by splitting
 void ProxyLOD::MixedPolyMeshToRawMesh(const FMixedPolyMesh& SimpleMesh, FMeshDescription& DstRawMesh)
 {
+	ResetStaticMeshDescription(DstRawMesh);
+
 	TVertexAttributesRef<FVector> VertexPositions = DstRawMesh.VertexAttributes().GetAttributesRef<FVector>(MeshAttribute::Vertex::Position);
 	TEdgeAttributesRef<bool> EdgeHardnesses = DstRawMesh.EdgeAttributes().GetAttributesRef<bool>(MeshAttribute::Edge::IsHard);
 	TEdgeAttributesRef<float> EdgeCreaseSharpnesses = DstRawMesh.EdgeAttributes().GetAttributesRef<float>(MeshAttribute::Edge::CreaseSharpness);
@@ -120,6 +130,9 @@ void ProxyLOD::MixedPolyMeshToRawMesh(const FMixedPolyMesh& SimpleMesh, FMeshDes
 
 void ProxyLOD::AOSMeshToRawMesh(const FAOSMesh& AOSMesh, FMeshDescription& OutRawMesh)
 {
+
+	ResetStaticMeshDescription(OutRawMesh);
+
 	TVertexAttributesRef<FVector> VertexPositions = OutRawMesh.VertexAttributes().GetAttributesRef<FVector>(MeshAttribute::Vertex::Position);
 	TEdgeAttributesRef<bool> EdgeHardnesses = OutRawMesh.EdgeAttributes().GetAttributesRef<bool>(MeshAttribute::Edge::IsHard);
 	TEdgeAttributesRef<float> EdgeCreaseSharpnesses = OutRawMesh.EdgeAttributes().GetAttributesRef<float>(MeshAttribute::Edge::CreaseSharpness);
@@ -209,6 +222,8 @@ void ProxyLOD::AOSMeshToRawMesh(const FAOSMesh& AOSMesh, FMeshDescription& OutRa
 
 void ProxyLOD::VertexDataMeshToRawMesh(const FVertexDataMesh& SrcVertexDataMesh, FMeshDescription& OutRawMesh)
 {
+	ResetStaticMeshDescription(OutRawMesh);
+
 	TVertexAttributesRef<FVector> VertexPositions = OutRawMesh.VertexAttributes().GetAttributesRef<FVector>(MeshAttribute::Vertex::Position);
 	TEdgeAttributesRef<bool> EdgeHardnesses = OutRawMesh.EdgeAttributes().GetAttributesRef<bool>(MeshAttribute::Edge::IsHard);
 	TEdgeAttributesRef<float> EdgeCreaseSharpnesses = OutRawMesh.EdgeAttributes().GetAttributesRef<float>(MeshAttribute::Edge::CreaseSharpness);
