Remove extra UV sets from HLOD meshes

* HLOD0 (merged) - Reliance on DoesUVChannelContainData() & the FMeshMergeSettings::OutputUVs[] array was lost during what looks like a bad code merge in ~2018
* HLOD1 (simplified) - Removed strange processing that added a second UV set, supposedly for lightmapping.  This isn't needed as it's handled in the static mesh build.

Bumped the HLOD proxy key to force a rebuild of all HLODs

diff --git a/Engine/Plugins/Experimental/ProxyLODPlugin/Source/ProxyLOD/Private/ProxyLODMeshConvertUtils.cpp b/Engine/Plugins/Experimental/ProxyLODPlugin/Source/ProxyLOD/Private/ProxyLODMeshConvertUtils.cpp
--- a/Engine/Plugins/Experimental/ProxyLODPlugin/Source/ProxyLOD/Private/ProxyLODMeshConvertUtils.cpp
+++ b/Engine/Plugins/Experimental/ProxyLODPlugin/Source/ProxyLOD/Private/ProxyLODMeshConvertUtils.cpp
@@ -240,12 +240,6 @@ void ProxyLOD::VertexDataMeshToRawMesh(const FVertexDataMesh& SrcVertexDataMesh,
 	const uint32 DstNumIndexes = SrcVertexDataMesh.Indices.Num();
 	const uint32 SrcNumTriangles = DstNumIndexes / 3;
 	
-	if (VertexInstanceUVs.GetNumIndices() < 2)
-	{
-		//We set the lightmap channel so 2
-		VertexInstanceUVs.SetNumIndices(2);
-	}
-
 	FPolygonGroupID PolygonGroupID = FPolygonGroupID::Invalid;
 	if (OutRawMesh.PolygonGroups().Num() == 0)
 	{
@@ -318,21 +312,17 @@ void ProxyLOD::VertexDataMeshToRawMesh(const FVertexDataMesh& SrcVertexDataMesh,
 			}
 			else
 			{
-				//There is 
 				VertexInstanceColors[VertexInstanceIDs[Corner]] = FVector4(FLinearColor(SrcVertexDataMesh.FaceColors[TriangleIndex]));
 			}
 
-			//UVs, copy two time the same value, one for UV_Channel 0 and another for lightmap channel 1
-			for (int32 channel = 0; channel < 2; ++channel)
+			//UVs
+			if (SrcVertexDataMesh.UVs.Num() == 0)
 			{
-				if (SrcVertexDataMesh.UVs.Num() == 0)
-				{
-					VertexInstanceUVs.Set(VertexInstanceIDs[Corner], channel, FVector2D(0.0f, 0.0f));
-				}
-				else
-				{
-					VertexInstanceUVs.Set(VertexInstanceIDs[Corner], channel, SrcVertexDataMesh.UVs[SrcIndex]);
-				}
+				VertexInstanceUVs.Set(VertexInstanceIDs[Corner], FVector2D(0.0f, 0.0f));
+			}
+			else
+			{
+				VertexInstanceUVs.Set(VertexInstanceIDs[Corner], SrcVertexDataMesh.UVs[SrcIndex]);
 			}
 		}
 
diff --git a/Engine/Source/Developer/MeshMergeUtilities/Private/MeshMergeUtilities.cpp b/Engine/Source/Developer/MeshMergeUtilities/Private/MeshMergeUtilities.cpp
--- a/Engine/Source/Developer/MeshMergeUtilities/Private/MeshMergeUtilities.cpp
+++ b/Engine/Source/Developer/MeshMergeUtilities/Private/MeshMergeUtilities.cpp
@@ -2928,6 +2928,10 @@ void FMeshMergeUtilities::CreateMergedRawMeshes(FMeshMergeDataTracker& InDataTra
 					});
 					AppendSettings.bMergeVertexColor = InSettings.bBakeVertexDataToMesh;
 					AppendSettings.MergedAssetPivot = InMergedAssetPivot;
+					for (int32 ChannelIdx = 0; ChannelIdx < FMeshDescriptionOperations::FAppendSettings::MAX_MESH_TEXTURE_COORDS; ++ChannelIdx)
+					{
+						AppendSettings.bMergeUVChannels[ChannelIdx] = InDataTracker.DoesUVChannelContainData(ChannelIdx, LODIndex) && InSettings.OutputUVs[ChannelIdx] == EUVOutput::OutputChannel;
+					}
 					FStaticMeshOperations::AppendMeshDescription(*RawMeshPtr, MergedMesh, AppendSettings);
 				}
 			}
@@ -3033,6 +3037,10 @@ void FMeshMergeUtilities::CreateMergedRawMeshes(FMeshMergeDataTracker& InDataTra
 				});
 				AppendSettings.bMergeVertexColor = InSettings.bBakeVertexDataToMesh;
 				AppendSettings.MergedAssetPivot = InMergedAssetPivot;
+				for (int32 ChannelIdx = 0; ChannelIdx < FMeshDescriptionOperations::FAppendSettings::MAX_MESH_TEXTURE_COORDS; ++ChannelIdx)
+				{
+					AppendSettings.bMergeUVChannels[ChannelIdx] = InDataTracker.DoesUVChannelContainData(ChannelIdx, LODIndex) && InSettings.OutputUVs[ChannelIdx] == EUVOutput::OutputChannel;
+				}
 				FStaticMeshOperations::AppendMeshDescription(*RawMeshPtr, MergedMesh, AppendSettings);
 			}
 		}
diff --git a/Engine/Source/Runtime/Engine/Private/HLODProxy.cpp b/Engine/Source/Runtime/Engine/Private/HLODProxy.cpp
--- a/Engine/Source/Runtime/Engine/Private/HLODProxy.cpp
+++ b/Engine/Source/Runtime/Engine/Private/HLODProxy.cpp
@@ -271,7 +271,7 @@ uint32 UHLODProxy::GetCRC(UStaticMeshComponent* InComponent, uint32 InCRC, const
 }
 
 // Key that forms the basis of the HLOD proxy key. Bump this key (i.e. generate a new GUID) when you want to force a rebuild of ALL HLOD proxies
-#define HLOD_PROXY_BASE_KEY		TEXT("76927B120C6645ACB9200E7FB8896AC3")
+#define HLOD_PROXY_BASE_KEY		TEXT("3E11E1EE1D554C4CB1A1C8D91098A6A7")
 
 FName UHLODProxy::GenerateKeyForActor(const ALODActor* LODActor, bool bMustUndoLevelTransform)
 {
diff --git a/Engine/Source/Runtime/StaticMeshDescription/Private/StaticMeshOperations.cpp b/Engine/Source/Runtime/StaticMeshDescription/Private/StaticMeshOperations.cpp
--- a/Engine/Source/Runtime/StaticMeshDescription/Private/StaticMeshOperations.cpp
+++ b/Engine/Source/Runtime/StaticMeshDescription/Private/StaticMeshOperations.cpp
@@ -853,9 +853,18 @@ void FStaticMeshOperations::AppendMeshDescription(const FMeshDescription& Source
 	TargetMesh.ReserveNewEdges(SourceMesh.Edges().Num());
 	TargetMesh.ReserveNewPolygons(SourceMesh.Polygons().Num());
 
-	if (SourceVertexInstanceUVs.GetNumIndices() > TargetVertexInstanceUVs.GetNumIndices())
+	int32 NumSourceUVChannels = 0;
+	for (int32 ChannelIdx = 0; ChannelIdx < SourceVertexInstanceUVs.GetNumIndices(); ++ChannelIdx)
 	{
-		TargetVertexInstanceUVs.SetNumIndices(SourceVertexInstanceUVs.GetNumIndices());
+		if (AppendSettings.bMergeUVChannels[ChannelIdx])
+		{
+			NumSourceUVChannels = ChannelIdx + 1;
+		}
+	}
+
+	if (NumSourceUVChannels > TargetVertexInstanceUVs.GetNumIndices())
+	{
+		TargetVertexInstanceUVs.SetNumIndices(NumSourceUVChannels);
 	}
 
 	//PolygonGroups
@@ -938,7 +947,7 @@ void FStaticMeshOperations::AppendMeshDescription(const FMeshDescription& Source
 			TargetVertexInstanceColors[TargetVertexInstanceID] = SourceVertexInstanceColors[SourceVertexInstanceID];
 		}
 
-		for (int32 UVChannelIndex = 0; UVChannelIndex < SourceVertexInstanceUVs.GetNumIndices(); ++UVChannelIndex)
+		for (int32 UVChannelIndex = 0; UVChannelIndex < NumSourceUVChannels; ++UVChannelIndex)
 		{
 			TargetVertexInstanceUVs.Set(TargetVertexInstanceID, UVChannelIndex, SourceVertexInstanceUVs.Get(SourceVertexInstanceID, UVChannelIndex));
 		}
diff --git a/Engine/Source/Runtime/StaticMeshDescription/Public/StaticMeshOperations.h b/Engine/Source/Runtime/StaticMeshDescription/Public/StaticMeshOperations.h
--- a/Engine/Source/Runtime/StaticMeshDescription/Public/StaticMeshOperations.h
+++ b/Engine/Source/Runtime/StaticMeshDescription/Public/StaticMeshOperations.h
@@ -37,10 +37,18 @@ public:
 	{
 		FAppendSettings()
 			: bMergeVertexColor(true)
+			, bMergeUVChannels{ true }
 			, MergedAssetPivot(0.0f, 0.0f, 0.0f)
 		{}
+
+		enum
+		{
+			MAX_MESH_TEXTURE_COORDS = 8,
+		};
+
 		FAppendPolygonGroupsDelegate PolygonGroupsDelegate;
 		bool bMergeVertexColor;
+		bool bMergeUVChannels[MAX_MESH_TEXTURE_COORDS];
 		FVector MergedAssetPivot;
 		TOptional<FTransform> MeshTransform; // Apply a transformation on source mesh (see MeshTransform)
 	};
