Fixed crash when a mesh with no triangles is provided to the Proxy mesh generation

diff --git a/Engine/Plugins/Experimental/ProxyLODPlugin/Source/ProxyLOD/Private/ProxyLODMeshTypes.cpp b/Engine/Plugins/Experimental/ProxyLODPlugin/Source/ProxyLOD/Private/ProxyLODMeshTypes.cpp
--- a/Engine/Plugins/Experimental/ProxyLODPlugin/Source/ProxyLOD/Private/ProxyLODMeshTypes.cpp
+++ b/Engine/Plugins/Experimental/ProxyLODPlugin/Source/ProxyLOD/Private/ProxyLODMeshTypes.cpp
@@ -61,41 +61,54 @@ void FMeshDescriptionArrayAdapter::Construct(int32 MeshCount, TFunctionRef<const
 	PointCount = 0;
 	PolyCount = 0;
 
+	PolyOffsetArray.clear();
+	IndexBufferArray.clear();
+	RawMeshArray.clear();
+	MergeDataArray.clear();
+
+	PolyOffsetArray.reserve(MeshCount);
+	IndexBufferArray.reserve(MeshCount);
+	RawMeshArray.reserve(MeshCount);
+	MergeDataArray.reserve(MeshCount);
+
 	PolyOffsetArray.push_back(PolyCount);
 	for (int32 MeshIdx = 0; MeshIdx < MeshCount; ++MeshIdx)
 	{
 		const FMeshMergeData* MergeData = GetMeshFunction(MeshIdx);
 		FMeshDescription *RawMesh = MergeData->RawMesh;
 
-		PointCount += size_t(RawMesh->Vertices().Num());
 		// Sum up all the polys in this mesh.
 		int32 MeshPolyCount = RawMesh->Triangles().Num();
-
-		// Construct local index buffer for this mesh
-		IndexBufferArray.push_back(std::vector<FVertexInstanceID>());
-		std::vector<FVertexInstanceID>& IndexBuffer = IndexBufferArray[MeshIdx];
-		IndexBuffer.reserve(MeshPolyCount * 3);
-		for (const FTriangleID TriangleID : RawMesh->Triangles().GetElementIDs())
+		if (MeshPolyCount > 0)
 		{
-			IndexBuffer.push_back(RawMesh->GetTriangleVertexInstance(TriangleID, 0));
-			IndexBuffer.push_back(RawMesh->GetTriangleVertexInstance(TriangleID, 1));
-			IndexBuffer.push_back(RawMesh->GetTriangleVertexInstance(TriangleID, 2));
-		}
+			// Construct local index buffer for this mesh
+			IndexBufferArray.push_back(std::vector<FVertexInstanceID>());
+			std::vector<FVertexInstanceID>& IndexBuffer = IndexBufferArray.back();
+			IndexBuffer.reserve(MeshPolyCount * 3);
 
-		PolyCount += MeshPolyCount;
-		PolyOffsetArray.push_back(PolyCount);
-		RawMeshArray.push_back(RawMesh);
-		MergeDataArray.push_back(MergeData);
+			for (const FTriangleID TriangleID : RawMesh->Triangles().GetElementIDs())
+			{
+				IndexBuffer.push_back(RawMesh->GetTriangleVertexInstance(TriangleID, 0));
+				IndexBuffer.push_back(RawMesh->GetTriangleVertexInstance(TriangleID, 1));
+				IndexBuffer.push_back(RawMesh->GetTriangleVertexInstance(TriangleID, 2));
+			}
+
+			PointCount += size_t(RawMesh->Vertices().Num());
+			PolyCount += MeshPolyCount;
+			PolyOffsetArray.push_back(PolyCount);
+			RawMeshArray.push_back(RawMesh);
+			MergeDataArray.push_back(MergeData);
+		}
 	}
 
-	RawMeshArrayData.SetNumUninitialized(MeshCount);
+	RawMeshArrayData.SetNumUninitialized(MergeDataArray.size());
 
 	// The getter constructor is doing expensive operations
 	ParallelFor(
 		RawMeshArrayData.Num(),
 		[this, &GetMeshFunction](int32 MeshIdx)
 		{
-			const FMeshMergeData* MergeData = GetMeshFunction(MeshIdx);
+			const FMeshMergeData* MergeData = MergeDataArray[MeshIdx];
 			FMeshDescription *RawMesh = MergeData->RawMesh;
 			new (&RawMeshArrayData[MeshIdx]) FMeshDescriptionAttributesGetter(RawMesh);
 		},
@@ -302,22 +315,9 @@ FMeshDescriptionArrayAdapter::FRawPoly FMeshDescriptionArrayAdapter::GetRawPoly(
 
 const FMeshDescription& FMeshDescriptionArrayAdapter::GetRawMesh(const size_t FaceNumber, int32& MeshIdx, int32& LocalFaceNumber, const FMeshDescriptionAttributesGetter** OutAttributesGetter) const
 {
-	// Find the correct raw mesh
-	MeshIdx = 0;
-
 	// Binary Search into poly offset will help a lot for big meshes
-	auto it = std::lower_bound(PolyOffsetArray.begin(), PolyOffsetArray.end(), FaceNumber);
-	if (it != PolyOffsetArray.end())
-	{
-		MeshIdx = std::distance(PolyOffsetArray.begin(), it);
-
-		// lower_bound will always get us a little bit farther than we wish
-		// we might need to backtrack an iteration or two.
-		while (FaceNumber < PolyOffsetArray[MeshIdx])
-		{
-			MeshIdx--;
-		}
-	}
+	auto it = std::upper_bound(PolyOffsetArray.begin(), PolyOffsetArray.end(), FaceNumber);
+	MeshIdx = std::distance(PolyOffsetArray.begin(), it) -  1;
 
 	// Offset the face number to get the correct index into this mesh.
 	LocalFaceNumber = int32(FaceNumber) - PolyOffsetArray[MeshIdx];
