Disabled ProxyLODMeshReduction plugin by default; note that in many cases the plugin is enabled through the ModelingToolsEditorMode plugin.

diff --git a/Engine/Plugins/Experimental/ProxyLODPlugin/ProxyLODPlugin.uplugin b/Engine/Plugins/Experimental/ProxyLODPlugin/ProxyLODPlugin.uplugin
--- a/Engine/Plugins/Experimental/ProxyLODPlugin/ProxyLODPlugin.uplugin
+++ b/Engine/Plugins/Experimental/ProxyLODPlugin/ProxyLODPlugin.uplugin
@@ -10,7 +10,7 @@
   "DocsURL" : "",
   "MarketplaceURL" : "",
   "SupportURL" : "",
-  "EnabledByDefault" : true,
+  "EnabledByDefault" : false,
   "CanContainContent" : false,
   "IsBetaVersion" : true,
   "Installed" : false,
