Fixed UE4 references in ProxyLODPlugin

diff --git a/Engine/Plugins/Experimental/ProxyLODPlugin/Source/ProxyLOD/Private/ProxyLODMeshUtilities.h b/Engine/Plugins/Experimental/ProxyLODPlugin/Source/ProxyLOD/Private/ProxyLODMeshUtilities.h
--- a/Engine/Plugins/Experimental/ProxyLODPlugin/Source/ProxyLOD/Private/ProxyLODMeshUtilities.h
+++ b/Engine/Plugins/Experimental/ProxyLODPlugin/Source/ProxyLOD/Private/ProxyLODMeshUtilities.h
@@ -31,7 +31,7 @@ namespace ProxyLOD
 	void ComputeVertexNormals( FVertexDataMesh& InOutMesh, 
 		                       const ENormalComputationMethod Method = ENormalComputationMethod::AngleWeighted);
 	/**
-	* Uses UE4 code in the Mesh Utilities module to compute the tangent space
+	* Uses UE code in the Mesh Utilities module to compute the tangent space
 	* The computation is a Mikk-T(angent) space version of tangent space.  
 	*
 	* NB: This updates the FMeshDescription::TangentX and FMeshDescription::TangentY and optionally TangentZ.
diff --git a/Engine/Plugins/Experimental/ProxyLODPlugin/Source/ProxyLOD/Private/ProxyLODPlugin.cpp b/Engine/Plugins/Experimental/ProxyLODPlugin/Source/ProxyLOD/Private/ProxyLODPlugin.cpp
--- a/Engine/Plugins/Experimental/ProxyLODPlugin/Source/ProxyLOD/Private/ProxyLODPlugin.cpp
+++ b/Engine/Plugins/Experimental/ProxyLODPlugin/Source/ProxyLOD/Private/ProxyLODPlugin.cpp
@@ -496,7 +496,7 @@ static void SimplifyMesh( const FClosestPolyField& SrcGeometryPolyField,
 {
 	TRACE_CPUPROFILER_EVENT_SCOPE(SimplifyMesh)
 
-	//SCOPE_LOG_TIME(TEXT("UE4_ProxyLOD_Simplifier"), nullptr);
+	//SCOPE_LOG_TIME(TEXT("UE_ProxyLOD_Simplifier"), nullptr);
 
 	// Compute some of the metrics that relate the desired resolution to simplifier parameters.
 	const FMeshDescriptionArrayAdapter& SrcGeometryAdapter = SrcGeometryPolyField.MeshAdapter();
diff --git a/Engine/Plugins/Experimental/ProxyLODPlugin/Source/ProxyLOD/Private/ProxyLODSimplifier.cpp b/Engine/Plugins/Experimental/ProxyLODPlugin/Source/ProxyLOD/Private/ProxyLODSimplifier.cpp
--- a/Engine/Plugins/Experimental/ProxyLODPlugin/Source/ProxyLOD/Private/ProxyLODSimplifier.cpp
+++ b/Engine/Plugins/Experimental/ProxyLODPlugin/Source/ProxyLOD/Private/ProxyLODSimplifier.cpp
@@ -24,7 +24,7 @@ ProxyLOD::FQuadricMeshSimplifier::FQuadricMeshSimplifier(const MeshVertType* Ver
 	edgeHash(1 << FMath::Min(16u, FMath::FloorLog2(NumVerts))),
 	quadricCache(numSVerts, numSTris)
 {
-	//	SCOPE_LOG_TIME(TEXT("UE4_ProxyLOD_Simplifier_Constructor"), nullptr);
+	//	SCOPE_LOG_TIME(TEXT("UE_ProxyLOD_Simplifier_Constructor"), nullptr);
 	vertFlagLock = 0;
 	triFlagLock = 0;
 
@@ -382,7 +382,7 @@ void ProxyLOD::FQuadricMeshSimplifier::InitCosts()
 {
 	TRACE_CPUPROFILER_EVENT_SCOPE(ProxyLOD::FQuadricMeshSimplifier::InitCosts)
 
-	//	SCOPE_LOG_TIME(TEXT("UE4_ProxyLOD_Simplifier_InitCosts"), nullptr);
+	//	SCOPE_LOG_TIME(TEXT("UE_ProxyLOD_Simplifier_InitCosts"), nullptr);
 	for (int i = 0; i < edges.Num(); i++)
 	{
 		float cost = ComputeEdgeCollapseCost(&edges[i]);
