ProxyLODMeshReduction module now has explicit runtime dependency on tbb.dll. Fixes module dependency errors on load.

diff --git a/Engine/Plugins/Experimental/ProxyLODPlugin/Source/ProxyLOD/ProxyLODMeshReduction.Build.cs b/Engine/Plugins/Experimental/ProxyLODPlugin/Source/ProxyLOD/ProxyLODMeshReduction.Build.cs
--- a/Engine/Plugins/Experimental/ProxyLODPlugin/Source/ProxyLOD/ProxyLODMeshReduction.Build.cs
+++ b/Engine/Plugins/Experimental/ProxyLODPlugin/Source/ProxyLOD/ProxyLODMeshReduction.Build.cs
@@ -7,9 +7,8 @@ namespace UnrealBuildTool.Rules
 		public ProxyLODMeshReduction(ReadOnlyTargetRules Target) : base(Target)
 		{
 
-
-            // For boost:: and TBB:: code
-            bUseRTTI = true;
+			// For boost:: and TBB:: code
+			bUseRTTI = true;
 
             PublicIncludePaths.AddRange(
 				new string[] {
@@ -70,6 +69,13 @@ namespace UnrealBuildTool.Rules
 					// ... add any modules that your module loads dynamically here ...
 				}
 				);
+
+			// OpenVDB pulls in a dependency for tbb.dll, make sure we specify it as necessary so the dll is sure to be where we need it to be
+			if (Target.Platform == UnrealTargetPlatform.Win64)
+			{
+				string IntelTBBLibs = Target.UEThirdPartyBinariesDirectory + "Intel/TBB/";
+				RuntimeDependencies.Add("$(TargetOutputDir)/tbb.dll", IntelTBBLibs + "Win64/tbb.dll");
+			}
 		}
 	}
 }
