ProxyLOD - Added support for providing instances transforms rather than having to build a gigantic FMeshDescription that include all instances

diff --git a/Engine/Plugins/Experimental/ProxyLODPlugin/Source/ProxyLOD/Private/ProxyLODMaterialTransferUtilities.cpp b/Engine/Plugins/Experimental/ProxyLODPlugin/Source/ProxyLOD/Private/ProxyLODMaterialTransferUtilities.cpp
--- a/Engine/Plugins/Experimental/ProxyLODPlugin/Source/ProxyLOD/Private/ProxyLODMaterialTransferUtilities.cpp
+++ b/Engine/Plugins/Experimental/ProxyLODPlugin/Source/ProxyLOD/Private/ProxyLODMaterialTransferUtilities.cpp
@@ -24,8 +24,8 @@ namespace
 	float EvaluateHit(const FVector& HitPoint, const uint32 SrcPolyId, const FMeshDescriptionArrayAdapter& MeshAdapter,
 		ProxyLOD::FSrcMeshData& HitPayload, TestFunctorType& TestFunctor)
 	{
-		int32 MeshId; int32 LocalFaceNumber;
-		const FMeshDescriptionArrayAdapter::FRawPoly SrcPoly = MeshAdapter.GetRawPoly(SrcPolyId, MeshId, LocalFaceNumber, ERawPolyValues::WedgeTexCoords | ERawPolyValues::VertexPositions | ERawPolyValues::WedgeTangents);
+		int32 MeshId, InstanceIdx, LocalFaceNumber;
+		const FMeshDescriptionArrayAdapter::FRawPoly SrcPoly = MeshAdapter.GetRawPoly(SrcPolyId, MeshId, InstanceIdx, LocalFaceNumber, ERawPolyValues::WedgeTexCoords | ERawPolyValues::VertexPositions | ERawPolyValues::WedgeTangents);
 		int32 MaterialIdx = SrcPoly.FaceMaterialIndex;
 
 		FVector2D UVs[3];
@@ -1070,8 +1070,8 @@ namespace
 
 			auto ComputeSrcTangentSpace = [&SrcMeshAdapter](const int32 TriangleId, const ProxyLOD::DArray3d& BarycentricCoords, FTangentSpace& TangentSpace)
 			{
-				int32 MeshId; int32 LocalFaceNumber;
-				const FMeshDescriptionArrayAdapter::FRawPoly SrcPoly = SrcMeshAdapter.GetRawPoly(TriangleId, MeshId, LocalFaceNumber, ERawPolyValues::WedgeTangents);
+				int32 MeshId, InstanceIdx, LocalFaceNumber;
+				const FMeshDescriptionArrayAdapter::FRawPoly SrcPoly = SrcMeshAdapter.GetRawPoly(TriangleId, MeshId, InstanceIdx, LocalFaceNumber, ERawPolyValues::WedgeTangents);
 
 				// Get the tangent space:  We try to compute this in a way that is consistent with LocalVertexFactory.usf (CalcTangentToLocal).
 
diff --git a/Engine/Plugins/Experimental/ProxyLODPlugin/Source/ProxyLOD/Private/ProxyLODMeshTypes.cpp b/Engine/Plugins/Experimental/ProxyLODPlugin/Source/ProxyLOD/Private/ProxyLODMeshTypes.cpp
--- a/Engine/Plugins/Experimental/ProxyLODPlugin/Source/ProxyLOD/Private/ProxyLODMeshTypes.cpp
+++ b/Engine/Plugins/Experimental/ProxyLODPlugin/Source/ProxyLOD/Private/ProxyLODMeshTypes.cpp
@@ -56,7 +56,35 @@ void FMeshDescriptionAdapter::getIndexSpacePoint(size_t FaceNumber, size_t Corne
 	pos = Transform.worldToIndex(openvdb::Vec3d(Position.X, Position.Y, Position.Z));
 };
 
+void FMeshDescriptionArrayAdapter::SetupInstances(int32 MeshCount, TFunctionRef<const FInstancedMeshMergeData* (uint32 Index)> GetMeshFunction)
+{
+	InstancesTransformArray.Reserve(MeshCount);
+	InstancesAdjointTArray.Reserve(MeshCount);
+
+	for (int32 MeshIdx = 0; MeshIdx < MeshCount; ++MeshIdx)
+	{
+		const FInstancedMeshMergeData* InstancedMeshMergeData = GetMeshFunction(MeshIdx);
+
+		InstancesTransformArray.Add(InstancedMeshMergeData->InstanceTransforms);
+
+		// Cache transpose adjoint matrices
+		TArray<FMatrix>& AdjointTArray = InstancesAdjointTArray.Emplace_GetRef();
+		for (const FTransform& InstanceTransform : InstancedMeshMergeData->InstanceTransforms)
+		{
+			FMatrix Matrix = InstanceTransform.ToMatrixWithScale();
+			FMatrix AdjointT = Matrix.TransposeAdjoint();
+			AdjointT.RemoveScaling();
+			AdjointTArray.Add(AdjointT);
+		}
+	}
+}
+
 void FMeshDescriptionArrayAdapter::Construct(int32 MeshCount, TFunctionRef<const FMeshMergeData* (uint32 Index)> GetMeshFunction)
+{
+	Construct(MeshCount, GetMeshFunction, [](uint32 Index) { return 1; });
+}
+
+void FMeshDescriptionArrayAdapter::Construct(int32 MeshCount, TFunctionRef<const FMeshMergeData* (uint32 Index)> GetMeshFunction, TFunctionRef<int32 (uint32 Index)> GetInstanceCountFunction)
 {
 	PointCount = 0;
 	PolyCount = 0;
@@ -75,6 +103,10 @@ void FMeshDescriptionArrayAdapter::Construct(int32 MeshCount, TFunctionRef<const
 	for (int32 MeshIdx = 0; MeshIdx < MeshCount; ++MeshIdx)
 	{
 		const FMeshMergeData* MergeData = GetMeshFunction(MeshIdx);
+
+		int32 InstanceCount = GetInstanceCountFunction(MeshIdx);
+		check(InstanceCount > 0);
+
 		FMeshDescription *RawMesh = MergeData->RawMesh;
 
 		// Sum up all the polys in this mesh.
@@ -93,8 +125,9 @@ void FMeshDescriptionArrayAdapter::Construct(int32 MeshCount, TFunctionRef<const
 				IndexBuffer.push_back(RawMesh->GetTriangleVertexInstance(TriangleID, 2));
 			}
 
-			PointCount += size_t(RawMesh->Vertices().Num());
-			PolyCount += MeshPolyCount;
+			PointCount += size_t(RawMesh->Vertices().Num() * InstanceCount);
+			PolyCount += MeshPolyCount * InstanceCount;
+
 			PolyOffsetArray.push_back(PolyCount);
 			RawMeshArray.push_back(RawMesh);
 			MergeDataArray.push_back(MergeData);
@@ -120,6 +153,15 @@ void FMeshDescriptionArrayAdapter::Construct(int32 MeshCount, TFunctionRef<const
 }
 
 // --- FMeshDescriptionArrayAdapter ----
+FMeshDescriptionArrayAdapter::FMeshDescriptionArrayAdapter(const TArray<const FInstancedMeshMergeData*>& InMergeDataPtrArray)
+{
+	// Make a default transform.
+	Transform = openvdb::math::Transform::createLinearTransform(1.);
+
+	SetupInstances(InMergeDataPtrArray.Num(), [&InMergeDataPtrArray](uint32 Index) { return InMergeDataPtrArray[Index]; });
+	Construct(InMergeDataPtrArray.Num(), [&InMergeDataPtrArray](uint32 Index) { return InMergeDataPtrArray[Index]; }, [&InMergeDataPtrArray](uint32 Index) { return InMergeDataPtrArray[Index]->InstanceTransforms.Num(); });
+}
+
 FMeshDescriptionArrayAdapter::FMeshDescriptionArrayAdapter(const TArray<const FMeshMergeData*>& InMergeDataPtrArray)
 {
 	// Make a default transform.
@@ -136,6 +178,15 @@ FMeshDescriptionArrayAdapter::FMeshDescriptionArrayAdapter(const TArray<FMeshMer
 	Construct(InMergeDataArray.Num(), [&InMergeDataArray](uint32 Index) { return &InMergeDataArray[Index]; });
 }
 
+FMeshDescriptionArrayAdapter::FMeshDescriptionArrayAdapter(const TArray<FInstancedMeshMergeData>& InMergeDataArray)
+{
+	// Make a default transform.
+	Transform = openvdb::math::Transform::createLinearTransform(1.);
+
+	SetupInstances(InMergeDataArray.Num(), [&InMergeDataArray](uint32 Index) { return &InMergeDataArray[Index]; });
+	Construct(InMergeDataArray.Num(), [&InMergeDataArray](uint32 Index) { return &InMergeDataArray[Index]; }, [&InMergeDataArray](uint32 Index) { return InMergeDataArray[Index].InstanceTransforms.Num(); });
+}
+
 FMeshDescriptionArrayAdapter::FMeshDescriptionArrayAdapter(const TArray<FMeshMergeData>& InMergeDataArray, const openvdb::math::Transform::Ptr InTransform)
 	: Transform(InTransform)
 {
@@ -148,6 +199,8 @@ FMeshDescriptionArrayAdapter::FMeshDescriptionArrayAdapter(const FMeshDescriptio
 	RawMeshArray = other.RawMeshArray;
 	PolyOffsetArray = other.PolyOffsetArray;
 	MergeDataArray = other.MergeDataArray;
+	InstancesTransformArray = other.InstancesTransformArray;
+	InstancesAdjointTArray = other.InstancesAdjointTArray;
 
 	RawMeshArrayData.SetNumUninitialized(RawMeshArray.size());
 
@@ -178,12 +231,12 @@ FMeshDescriptionArrayAdapter::~FMeshDescriptionArrayAdapter()
 	PolyOffsetArray.clear();
 }
 
-void FMeshDescriptionArrayAdapter::getWorldSpacePoint(size_t FaceNumber, size_t CornerNumber, openvdb::Vec3d& pos) const
+void FMeshDescriptionArrayAdapter::GetWorldSpacePoint(size_t FaceNumber, size_t CornerNumber, openvdb::Vec3d& pos) const
 {
-	int32 MeshIdx, LocalFaceNumber;
+	int32 MeshIdx, InstanceIdx, LocalFaceNumber;
 
 	const FMeshDescriptionAttributesGetter* AttributesGetter = nullptr;
-	const FMeshDescription& RawMesh = GetRawMesh(FaceNumber, MeshIdx, LocalFaceNumber, &AttributesGetter);
+	const FMeshDescription& RawMesh = GetRawMesh(FaceNumber, MeshIdx, InstanceIdx, LocalFaceNumber, &AttributesGetter);
 	check(AttributesGetter);
 
 	const auto& IndexBuffer = IndexBufferArray[MeshIdx];
@@ -191,13 +244,19 @@ void FMeshDescriptionArrayAdapter::getWorldSpacePoint(size_t FaceNumber, size_t
 	const FVertexInstanceID VertexInstanceID = IndexBuffer[3 * LocalFaceNumber + int32(CornerNumber)];
 	// float3 position 
 	FVector Position = AttributesGetter->VertexPositions[RawMesh.GetVertexInstanceVertex(VertexInstanceID)];
+
+	if (InstanceIdx != INDEX_NONE)
+	{
+		Position = InstancesTransformArray[MeshIdx][InstanceIdx].TransformPosition(Position);
+	}
+
 	pos = openvdb::Vec3d(Position.X, Position.Y, Position.Z);
 };
 
 void FMeshDescriptionArrayAdapter::getIndexSpacePoint(size_t FaceNumber, size_t CornerNumber, openvdb::Vec3d& pos) const
 {
 	openvdb::Vec3d Position;
-	getWorldSpacePoint(FaceNumber, CornerNumber, Position);
+	GetWorldSpacePoint(FaceNumber, CornerNumber, Position);
 	pos = Transform->worldToIndex(Position);
 
 };
@@ -232,16 +291,17 @@ void FMeshDescriptionArrayAdapter::UpdateMaterialsID()
 	}
 }
 
-FMeshDescriptionArrayAdapter::FRawPoly FMeshDescriptionArrayAdapter::GetRawPoly(const size_t FaceNumber, int32& OutMeshIdx, int32& OutLocalFaceNumber, const ERawPolyValues RawPolyValues ) const
+FMeshDescriptionArrayAdapter::FRawPoly FMeshDescriptionArrayAdapter::GetRawPoly(const size_t FaceNumber, int32& OutMeshIdx, int32& OutInstanceIdx, int32& OutLocalFaceNumber, const ERawPolyValues RawPolyValues ) const
 {
 	checkSlow(FaceNumber < PolyCount);
 
-	int32 MeshIdx, LocalFaceNumber;
+	int32 MeshIdx, InstanceIdx, LocalFaceNumber;
 
 	const FMeshDescriptionAttributesGetter* AttributesGetter = nullptr;
-	const FMeshDescription& RawMesh = GetRawMesh(FaceNumber, MeshIdx, LocalFaceNumber, &AttributesGetter);
+	const FMeshDescription& RawMesh = GetRawMesh(FaceNumber, MeshIdx, InstanceIdx, LocalFaceNumber, &AttributesGetter);
 	check(AttributesGetter);
 	OutMeshIdx = MeshIdx;
+	OutInstanceIdx = InstanceIdx;
 	OutLocalFaceNumber = LocalFaceNumber;
 
 	checkSlow(LocalFaceNumber < AttributesGetter->TriangleCount);
@@ -252,6 +312,9 @@ FMeshDescriptionArrayAdapter::FRawPoly FMeshDescriptionArrayAdapter::GetRawPoly(
 	RawPoly.FaceMaterialIndex = RawMesh.GetPolygonPolygonGroup(PolygonID).GetValue();
 	RawPoly.FaceSmoothingMask = AttributesGetter->FaceSmoothingMasks[LocalFaceNumber];
 
+	const bool bIsMirrored = InstanceIdx != INDEX_NONE ? InstancesTransformArray[MeshIdx][InstanceIdx].GetDeterminant() < 0.f : false;
+	const float MulBy = bIsMirrored ? -1.f : 1.f;
+
 	for (const FTriangleID TriangleID : RawMesh.GetPolygonTriangles(PolygonID))
 	{
 		TArrayView<const FVertexInstanceID> VertexInstanceIDs = RawMesh.GetTriangleVertexInstances(TriangleID);
@@ -262,6 +325,14 @@ FMeshDescriptionArrayAdapter::FRawPoly FMeshDescriptionArrayAdapter::GetRawPoly(
 			{
 				RawPoly.VertexPositions[i] = AttributesGetter->VertexPositions[RawMesh.GetVertexInstanceVertex(VertexInstanceIDs[i])];
 			}
+
+			if (InstanceIdx != INDEX_NONE)
+			{
+				for (int32 i = 0; i < 3; ++i)
+				{
+					RawPoly.VertexPositions[i] = InstancesTransformArray[MeshIdx][InstanceIdx].TransformPosition(RawPoly.VertexPositions[i]);
+				}
+			}
 		}
 
 		if ((RawPolyValues & ERawPolyValues::WedgeTangents) == ERawPolyValues::WedgeTangents)
@@ -272,6 +343,16 @@ FMeshDescriptionArrayAdapter::FRawPoly FMeshDescriptionArrayAdapter::GetRawPoly(
 				RawPoly.WedgeTangentY[i] = FVector::CrossProduct(AttributesGetter->VertexInstanceNormals[VertexInstanceIDs[i]], AttributesGetter->VertexInstanceTangents[VertexInstanceIDs[i]]).GetSafeNormal() * AttributesGetter->VertexInstanceBinormalSigns[VertexInstanceIDs[i]];
 				RawPoly.WedgeTangentZ[i] = AttributesGetter->VertexInstanceNormals[VertexInstanceIDs[i]];
 			}
+
+			if (InstanceIdx != INDEX_NONE)
+			{
+				for (int32 i = 0; i < 3; ++i)
+				{
+					RawPoly.WedgeTangentX[i] = InstancesAdjointTArray[MeshIdx][InstanceIdx].TransformVector(RawPoly.WedgeTangentX[i]) * MulBy;
+					RawPoly.WedgeTangentY[i] = InstancesAdjointTArray[MeshIdx][InstanceIdx].TransformVector(RawPoly.WedgeTangentY[i]) * MulBy;
+					RawPoly.WedgeTangentZ[i] = InstancesAdjointTArray[MeshIdx][InstanceIdx].TransformVector(RawPoly.WedgeTangentZ[i]) * MulBy;
+				}
+			}
 		}
 
 		if ((RawPolyValues & ERawPolyValues::WedgeColors) == ERawPolyValues::WedgeColors)
@@ -307,13 +388,13 @@ FMeshDescriptionArrayAdapter::FRawPoly FMeshDescriptionArrayAdapter::GetRawPoly(
 
 FMeshDescriptionArrayAdapter::FRawPoly FMeshDescriptionArrayAdapter::GetRawPoly(const size_t FaceNumber, const ERawPolyValues RawPolyValues) const
 {
-	int32 IgnoreMeshId, IgnoreLocalFaceNumber;
-	return GetRawPoly(FaceNumber, IgnoreMeshId, IgnoreLocalFaceNumber, RawPolyValues);
+	int32 IgnoreMeshId, IgnoreInstanceIdx, IgnoreLocalFaceNumber;
+	return GetRawPoly(FaceNumber, IgnoreMeshId, IgnoreInstanceIdx, IgnoreLocalFaceNumber, RawPolyValues);
 }
 
 // protected functions
 
-const FMeshDescription& FMeshDescriptionArrayAdapter::GetRawMesh(const size_t FaceNumber, int32& MeshIdx, int32& LocalFaceNumber, const FMeshDescriptionAttributesGetter** OutAttributesGetter) const
+const FMeshDescription& FMeshDescriptionArrayAdapter::GetRawMesh(const size_t FaceNumber, int32& MeshIdx, int32& InstanceIdx, int32& LocalFaceNumber, const FMeshDescriptionAttributesGetter** OutAttributesGetter) const
 {
 	// Binary Search into poly offset will help a lot for big meshes
 	auto it = std::upper_bound(PolyOffsetArray.begin(), PolyOffsetArray.end(), FaceNumber);
@@ -322,6 +403,18 @@ const FMeshDescription& FMeshDescriptionArrayAdapter::GetRawMesh(const size_t Fa
 	// Offset the face number to get the correct index into this mesh.
 	LocalFaceNumber = int32(FaceNumber) - PolyOffsetArray[MeshIdx];
 
+	int32 InstanceCount = !InstancesTransformArray.IsEmpty() ? InstancesTransformArray[MeshIdx].Num() : INDEX_NONE;
+	if (InstanceCount != INDEX_NONE)
+	{
+		int32 NumFacesPerInstance = RawMeshArrayData[MeshIdx].TriangleCount;
+		InstanceIdx = LocalFaceNumber / NumFacesPerInstance;
+		LocalFaceNumber = LocalFaceNumber % NumFacesPerInstance;
+	}
+	else
+	{
+		InstanceIdx = INDEX_NONE;
+	}
+
 	const FMeshDescription* MeshDescription = RawMeshArray[MeshIdx];
 
 	*OutAttributesGetter = &RawMeshArrayData[MeshIdx];
@@ -342,7 +435,7 @@ void FMeshDescriptionArrayAdapter::ComputeAABB(ProxyLOD::FBBox& InOutBBox)
 			// loop over verts
 			for (int32 v = 0; v < 3; ++v)
 			{
-				this->getWorldSpacePoint(f, v, Pos);
+				this->GetWorldSpacePoint(f, v, Pos);
 
 				TargetBBox.expand(Pos);
 			}
diff --git a/Engine/Plugins/Experimental/ProxyLODPlugin/Source/ProxyLOD/Private/ProxyLODMeshTypes.h b/Engine/Plugins/Experimental/ProxyLODPlugin/Source/ProxyLOD/Private/ProxyLODMeshTypes.h
--- a/Engine/Plugins/Experimental/ProxyLODPlugin/Source/ProxyLOD/Private/ProxyLODMeshTypes.h
+++ b/Engine/Plugins/Experimental/ProxyLODPlugin/Source/ProxyLOD/Private/ProxyLODMeshTypes.h
@@ -272,11 +272,36 @@ ENUM_CLASS_FLAGS(ERawPolyValues);
 class FMeshDescriptionArrayAdapter
 {
 public:
+	// OpenVDB MeshDataAdapter Interface
+	
+	// Total number of polygons managed by this class.
+	size_t polygonCount() const
+	{
+		return PolyCount;
+	}
+
+	// Total number of points (positions) managed by this class
+	size_t pointCount() const
+	{
+		return PointCount;
+	}
+
+	// Vertex count for polygon n: currently FMeshDescription is just triangles.
+	size_t vertexCount(size_t n) const
+	{
+		return 3;
+	}
 
+	// Return position pos in local grid index space for polygon (face number) n and vertex (conrner number) v
+	void getIndexSpacePoint(size_t FaceNumber, size_t CornerNumber, openvdb::Vec3d& pos) const;
+
+public:
 	FMeshDescriptionArrayAdapter(const TArray<FMeshMergeData>& InMergeDataArray, const openvdb::math::Transform::Ptr InTransform);
 	FMeshDescriptionArrayAdapter(const TArray<const FMeshMergeData*>& InMergeDataPtrArray);
+	FMeshDescriptionArrayAdapter(const TArray<const FInstancedMeshMergeData*>& InMergeDataPtrArray);
 	FMeshDescriptionArrayAdapter(const TArray<FMeshMergeData>& InMergeDataArray);
-
+	FMeshDescriptionArrayAdapter(const TArray<FInstancedMeshMergeData>& InMergeDataArray);
+	
 	// copy constructor 
 	FMeshDescriptionArrayAdapter(const FMeshDescriptionArrayAdapter& other);
 
@@ -284,12 +309,7 @@ public:
 	~FMeshDescriptionArrayAdapter();
 	
 	// Return position for polygon (face number) n and vertex (corner number) v
-	void getWorldSpacePoint(size_t FaceNumber, size_t CornerNumber, openvdb::Vec3d& pos) const;
-
-
-	// Return position pos in local grid index space for polygon (face number) n and vertex (conrner number) v
-	void getIndexSpacePoint(size_t FaceNumber, size_t CornerNumber, openvdb::Vec3d& pos) const;
-	
+	void GetWorldSpacePoint(size_t FaceNumber, size_t CornerNumber, openvdb::Vec3d& pos) const;
 
 	// Access to the FMeshMergeData data elements in the array.
 	const FMeshMergeData& GetMeshMergeData(uint32 Idx) const;
@@ -356,12 +376,13 @@ public:
 	* 
 	* @param FaceNumber          The triangle Id when treating all the meshes as a single mesh
 	* @param OutMeshIdx          The Id of the actual mesh that owns this poly
+	* @param OutInstanceIdx      The instance index, if any, or INDEX_NONE
 	* @param OutLocalFaceNumber  The Id within that mesh of this poly
 	* @param RawPolyValues       Reduce computations by specifying which values will be used from FRawPoly.
 	*
 	* @return  A copy of the raw mesh data associated with this poly.
 	*/
-	FMeshDescriptionArrayAdapter::FRawPoly GetRawPoly(const size_t FaceNumber, int32& OutMeshIdx, int32& OutLocalFaceNumber, const ERawPolyValues RawPolyValues = ERawPolyValues::All ) const;
+	FMeshDescriptionArrayAdapter::FRawPoly GetRawPoly(const size_t FaceNumber, int32& OutMeshIdx, int32& OutInstanceIdx, int32& OutLocalFaceNumber, const ERawPolyValues RawPolyValues = ERawPolyValues::All ) const;
 	
 	/**
 	* Returns a copy of the data associated with this poly in the form of a struct.
@@ -373,30 +394,6 @@ public:
 	*/
 	FMeshDescriptionArrayAdapter::FRawPoly GetRawPoly(const size_t FaceNumber, const ERawPolyValues RawPolyValues = ERawPolyValues::All ) const;
 
-	/**
-	* Total number of polygons managed by this class.
-	*/
-	size_t polygonCount() const 
-	{ 
-		return PolyCount; 
-	}
-
-	/** 
-	* Total number of points (positions) managed by this class
-	*/
-	size_t pointCount() const 
-	{ 
-		return PointCount; 
-	}
-
-	/**
-	* Vertex count for polygon n: currently FMeshDescription is just triangles.
-	*/
-	size_t vertexCount(size_t n) const 
-	{ 
-		return 3; 
-	}
-
 	/**
 	* The transform used to map between the physical space of the mesh and the voxel space.
 	*/
@@ -418,9 +415,12 @@ public:
 	}
 
 protected:
+	void SetupInstances(int32 MeshCount, TFunctionRef<const FInstancedMeshMergeData* (uint32 Index)> GetMeshFunction);
+
 	void Construct(int32 MeshCount, TFunctionRef<const FMeshMergeData* (uint32 Index)> GetMeshFunction);
+	void Construct(int32 MeshCount, TFunctionRef<const FMeshMergeData* (uint32 Index)> GetMeshFunction, TFunctionRef<int32(uint32 Index)> GetInstanceCountFunction);
 
-	const FMeshDescription& GetRawMesh(const size_t FaceNumber, int32& MeshIdx, int32& LocalFaceNumber, const FMeshDescriptionAttributesGetter** OutAttributesGetter) const;
+	const FMeshDescription& GetRawMesh(const size_t FaceNumber, int32& MeshIdx, int32& InstanceIdx, int32& LocalFaceNumber, const FMeshDescriptionAttributesGetter** OutAttributesGetter) const;
 
 	void ComputeAABB(ProxyLOD::FBBox& InOutBBox);
 protected:
@@ -434,6 +434,8 @@ protected:
 
 	std::vector<size_t>                      PolyOffsetArray;
 	std::vector<FMeshDescription*>           RawMeshArray;
+	TArray<TArray<FTransform>>               InstancesTransformArray;
+	TArray<TArray<FMatrix>>                  InstancesAdjointTArray;
 
 	// Use TArray because we need SetNumUninitialized
 	TArray<FMeshDescriptionAttributesGetter> RawMeshArrayData;
diff --git a/Engine/Plugins/Experimental/ProxyLODPlugin/Source/ProxyLOD/Private/ProxyLODPlugin.cpp b/Engine/Plugins/Experimental/ProxyLODPlugin/Source/ProxyLOD/Private/ProxyLODPlugin.cpp
--- a/Engine/Plugins/Experimental/ProxyLODPlugin/Source/ProxyLOD/Private/ProxyLODPlugin.cpp
+++ b/Engine/Plugins/Experimental/ProxyLODPlugin/Source/ProxyLOD/Private/ProxyLODPlugin.cpp
@@ -138,9 +138,10 @@ class FVoxelizeMeshMerging : public IMeshMerging
 {
 public:
 
-	typedef TArray<FMeshMergeData>    FMeshMergeDataArray;
-	typedef TArray<FFlattenMaterial>  FFlattenMaterialArray;
-	typedef openvdb::math::Transform  OpenVDBTransform;
+	typedef TArray<FMeshMergeData>			FMeshMergeDataArray;
+	typedef TArray<FInstancedMeshMergeData> FInstancedMeshMergeDataArray;
+	typedef TArray<FFlattenMaterial>		FFlattenMaterialArray;
+	typedef openvdb::math::Transform		OpenVDBTransform;
 
 	FVoxelizeMeshMerging();
 
@@ -154,6 +155,11 @@ public:
 		                  const FFlattenMaterialArray& InputMaterials,
 		                  const FGuid InJobGUID) override;
 
+	virtual void ProxyLOD(const FInstancedMeshMergeDataArray& InData,
+		const FMeshProxySettings& InProxySettings,
+		const FFlattenMaterialArray& InputMaterials,
+		const FGuid InJobGUID) override;
+
 	virtual void AggregateLOD() override
 	{};
 
@@ -170,6 +176,7 @@ public:
 
 
 private:
+	void ProxyLOD(FMeshDescriptionArrayAdapter& InSrcGeometryAdapter, FMeshDescriptionArrayAdapter& InClippingGeometryAdapter, const FMeshProxySettings& InProxySettings, const FFlattenMaterialArray& InputMaterials, const FGuid InJobGUID);
 
 	// Restore default parameters
 	void RestoreDefaultParameters();
@@ -628,11 +635,6 @@ static ProxyLOD::ENormalComputationMethod GetNormalComputationMethod(const FMesh
 
 void FVoxelizeMeshMerging::ProxyLOD(const FMeshMergeDataArray& InData, const FMeshProxySettings& InProxySettings, const FFlattenMaterialArray& InputMaterials, const FGuid InJobGUID)
 {
-	TRACE_CPUPROFILER_EVENT_SCOPE(FVoxelizeMeshMerging::ProxyLOD)
-	// Update any pipeline controlling parameters.
-
-	CaptureCVars();
-
 	// Split the input meshes into two groups.  The main geometry and clipping geometry.
 	// NB: only use pointers to avoid potentially copying a TArray of UV data.
 
@@ -641,7 +643,7 @@ void FVoxelizeMeshMerging::ProxyLOD(const FMeshMergeDataArray& InData, const FMe
 
 	for (int32 i = 0; i < InData.Num(); ++i)
 	{
-		const FMeshMergeData&  MeshMergeData = InData[i];
+		const FMeshMergeData& MeshMergeData = InData[i];
 		if (MeshMergeData.bIsClippingMesh)
 		{
 			InClippingGeometry.Add(&MeshMergeData);
@@ -652,7 +654,49 @@ void FVoxelizeMeshMerging::ProxyLOD(const FMeshMergeDataArray& InData, const FMe
 		}
 	}
 
-	if (InGeometry.Num() == 0)
+	// Create an adapter to make the data appear as a single mesh as required by the voxelization code.
+	FMeshDescriptionArrayAdapter SrcGeometryAdapter(InGeometry);
+	FMeshDescriptionArrayAdapter ClippingGeometryAdapter(InClippingGeometry);
+
+	ProxyLOD(SrcGeometryAdapter, ClippingGeometryAdapter, InProxySettings, InputMaterials, InJobGUID);
+}
+
+void FVoxelizeMeshMerging::ProxyLOD(const FInstancedMeshMergeDataArray& InData, const FMeshProxySettings& InProxySettings, const FFlattenMaterialArray& InputMaterials, const FGuid InJobGUID)
+{
+	// Split the input meshes into two groups.  The main geometry and clipping geometry.
+	// NB: only use pointers to avoid potentially copying a TArray of UV data.
+
+	TArray<const FInstancedMeshMergeData*> InGeometry;
+	TArray<const FInstancedMeshMergeData*> InClippingGeometry;
+
+	for (int32 i = 0; i < InData.Num(); ++i)
+	{
+		const FInstancedMeshMergeData& MeshMergeData = InData[i];
+		if (MeshMergeData.bIsClippingMesh)
+		{
+			InClippingGeometry.Add(&MeshMergeData);
+		}
+		else
+		{
+			InGeometry.Add(&MeshMergeData);
+		}
+	}
+
+	// Create an adapter to make the data appear as a single mesh as required by the voxelization code.
+	FMeshDescriptionArrayAdapter SrcGeometryAdapter(InGeometry);
+	FMeshDescriptionArrayAdapter ClippingGeometryAdapter(InClippingGeometry);
+
+	ProxyLOD(SrcGeometryAdapter, ClippingGeometryAdapter, InProxySettings, InputMaterials, InJobGUID);
+}
+
+void FVoxelizeMeshMerging::ProxyLOD(FMeshDescriptionArrayAdapter& InSrcGeometryAdapter, FMeshDescriptionArrayAdapter& InClippingGeometryAdapter, const FMeshProxySettings& InProxySettings, const FFlattenMaterialArray& InputMaterials, const FGuid InJobGUID)
+{
+	TRACE_CPUPROFILER_EVENT_SCOPE(FVoxelizeMeshMerging::ProxyLOD)
+	// Update any pipeline controlling parameters.
+
+	CaptureCVars();
+
+	if (InSrcGeometryAdapter.polygonCount() == 0)
 	{
 		UE_LOG(LogProxyLODMeshReduction, Warning, TEXT("No static meshes input, no output mesh will be generated."));
 
@@ -685,13 +729,8 @@ void FVoxelizeMeshMerging::ProxyLOD(const FMeshMergeDataArray& InData, const FMe
 	{
 		TRACE_CPUPROFILER_EVENT_SCOPE(ComputeSimplifiedMesh)
 
-		// Create an adapter to make the data appear as a single mesh as required by the voxelization code.
-
-		FMeshDescriptionArrayAdapter SrcGeometryAdapter(InGeometry);
-		FMeshDescriptionArrayAdapter ClippingGeometryAdapter(InClippingGeometry);
-
 		{
-			const auto& BBox = SrcGeometryAdapter.GetBBox();
+			const auto& BBox = InSrcGeometryAdapter.GetBBox();
 
 			if (!BBox.hasVolume())
 			{
@@ -713,20 +752,20 @@ void FVoxelizeMeshMerging::ProxyLOD(const FMeshMergeDataArray& InData, const FMe
 			// The transform defines resolution of the voxel grid.
 
 			OpenVDBTransform::Ptr XForm = ComputeResolution(InProxySettings, BBoxMajorAxisLength(BBox) );
-			SrcGeometryAdapter.SetTransform(XForm);
-			ClippingGeometryAdapter.SetTransform(XForm);
+			InSrcGeometryAdapter.SetTransform(XForm);
+			InClippingGeometryAdapter.SetTransform(XForm);
 		}
 
-		const double VoxelSize = SrcGeometryAdapter.GetTransform().voxelSize()[0];
+		const double VoxelSize = InSrcGeometryAdapter.GetTransform().voxelSize()[0];
 
 		// Prepare kDOPTree asynchronously ahead of time so it's ready when we need it
 		ProxyLOD::FkDOPTree  kDOPTree;
 		ProxyLOD::FTaskGroup kDOPTaskGroup;
 
 		kDOPTaskGroup.Run(
-			[&kDOPTree, &SrcGeometryAdapter]()
+			[&kDOPTree, &InSrcGeometryAdapter]()
 			{
-				ProxyLOD::BuildkDOPTree(SrcGeometryAdapter, kDOPTree);
+				ProxyLOD::BuildkDOPTree(InSrcGeometryAdapter, kDOPTree);
 			}
 		);
 
@@ -775,15 +814,15 @@ void FVoxelizeMeshMerging::ProxyLOD(const FMeshMergeDataArray& InData, const FMe
 
 				// 1) Voxelize - this can potentially run out of memory when very large objecs (or very small voxel sizes) are used.
 				
-				const bool bSuccess = ProxyLOD::MeshArrayToSDFVolume(SrcGeometryAdapter, SDFVolume, SrcPolyIndexGrid.get());
-				const bool bHasClipping = (InClippingGeometry.Num() != 0);
+				const bool bSuccess = ProxyLOD::MeshArrayToSDFVolume(InSrcGeometryAdapter, SDFVolume, SrcPolyIndexGrid.get());
+				const bool bHasClipping = InClippingGeometryAdapter.polygonCount() != 0;
 
 				if (bSuccess && bHasClipping)
 				{
 					// Voxelize the clipping geometry.
 
 					openvdb::FloatGrid::Ptr SDFClipping;
-					ProxyLOD::MeshArrayToSDFVolume(ClippingGeometryAdapter, SDFClipping);
+					ProxyLOD::MeshArrayToSDFVolume(InClippingGeometryAdapter, SDFClipping);
 
 					// CSG difference that removes the clipping region from the SDFvolume, leaving a watertight SDF
 
@@ -832,7 +871,7 @@ void FVoxelizeMeshMerging::ProxyLOD(const FMeshMergeDataArray& InData, const FMe
 
 				// 3) Create an object that allows for closest poly queries against the source mesh.
 
-				SrcGeometryPolyField = FClosestPolyField::create(SrcGeometryAdapter, SrcPolyIndexGrid);
+				SrcGeometryPolyField = FClosestPolyField::create(InSrcGeometryAdapter, SrcPolyIndexGrid);
 
 				// Let the SDFVolume go out of scope, no longer needed.
 			}
@@ -959,7 +998,7 @@ void FVoxelizeMeshMerging::ProxyLOD(const FMeshMergeDataArray& InData, const FMe
 					}
 				}
 			);
-			SrcGeometryAdapter.UpdateMaterialsID();
+			InSrcGeometryAdapter.UpdateMaterialsID();
 		}
 		else
 		{
@@ -1067,7 +1106,7 @@ void FVoxelizeMeshMerging::ProxyLOD(const FMeshMergeDataArray& InData, const FMe
 				// distance between polygons in the simplified geometry and polygons in the source geometry.
 				const bool bOverrideRayDist = InProxySettings.bOverrideTransferDistance;
 				double RequestedMaxRay = (bOverrideRayDist) ? InProxySettings.MaxRayCastDist : 6. * VoxelSize;
-				const double BBoxMajorAxisLength = SrcGeometryAdapter.GetBBox().extents()[SrcGeometryAdapter.GetBBox().maxExtent()];
+				const double BBoxMajorAxisLength = InSrcGeometryAdapter.GetBBox().extents()[InSrcGeometryAdapter.GetBBox().maxExtent()];
 				const double MaxRayLength = FMath::Max(VoxelSize, FMath::Min(RequestedMaxRay, BBoxMajorAxisLength));
 
 				UE_LOG(LogProxyLODMeshReduction, Log, TEXT("Material Distance %f used to discover textures."), float(MaxRayLength));
@@ -1079,7 +1118,7 @@ void FVoxelizeMeshMerging::ProxyLOD(const FMeshMergeDataArray& InData, const FMe
 				kDOPTaskGroup.Wait();
 
 				ProxyLOD::FSrcDataGrid::Ptr SuperSampledCorrespondenceGrid =
-					ProxyLOD::CreateCorrespondence(SrcGeometryAdapter, kDOPTree, VertexDataMesh, *DstSuperSampledUVGrid, this->RayHitOrder, MaxRayLength);
+					ProxyLOD::CreateCorrespondence(InSrcGeometryAdapter, kDOPTree, VertexDataMesh, *DstSuperSampledUVGrid, this->RayHitOrder, MaxRayLength);
 
 				// just for testing, this should force it to save in world space
 				const bool bUseWorldSpaceNormals = (!this->bUseTangentSpace);
@@ -1094,7 +1133,7 @@ void FVoxelizeMeshMerging::ProxyLOD(const FMeshMergeDataArray& InData, const FMe
 				// --- to generate flattened materials for the simplified geometry                       ---
 				// Compute the baked-down maps
 
-				ProxyLOD::MapFlattenMaterials(OutRawMesh, SrcGeometryAdapter, *SuperSampledCorrespondenceGrid, *DstSuperSampledUVGrid, *DstUVGrid, *BakedMaterials, UnresolvedGeometryColor, OutMaterial);
+				ProxyLOD::MapFlattenMaterials(OutRawMesh, InSrcGeometryAdapter, *SuperSampledCorrespondenceGrid, *DstSuperSampledUVGrid, *DstUVGrid, *BakedMaterials, UnresolvedGeometryColor, OutMaterial);
 
 				// Transfer the vertex colors unless we want to display the Charts as vertex colors for debugging
 
diff --git a/Engine/Plugins/Experimental/ProxyLODPlugin/Source/ProxyLOD/Private/ProxyLODVolume.cpp b/Engine/Plugins/Experimental/ProxyLODPlugin/Source/ProxyLOD/Private/ProxyLODVolume.cpp
--- a/Engine/Plugins/Experimental/ProxyLODPlugin/Source/ProxyLOD/Private/ProxyLODVolume.cpp
+++ b/Engine/Plugins/Experimental/ProxyLODPlugin/Source/ProxyLOD/Private/ProxyLODVolume.cpp
@@ -41,22 +41,13 @@ public:
 	bool Initialize(const TArray<FMeshMergeData>& Geometry, float Accuracy)
 	{
 		FMeshDescriptionArrayAdapter SrcGeometryAdapter(Geometry);
-		OpenVDBTransform::Ptr XForm = OpenVDBTransform::createLinearTransform(Accuracy);
-		SrcGeometryAdapter.SetTransform(XForm);
-
-		VoxelSize = SrcGeometryAdapter.GetTransform().voxelSize()[0];
-
-		SrcPolyIndexGrid = openvdb::Int32Grid::create();
-
-		if (!ProxyLOD::MeshArrayToSDFVolume(SrcGeometryAdapter, SDFVolume, SrcPolyIndexGrid.get()))
-		{
-			SrcPolyIndexGrid.reset();
-			return false;
-		}
-
-		Sampler.reset(new openvdb::tools::GridSampler<openvdb::FloatGrid, openvdb::tools::PointSampler>(*SDFVolume));
+		return Initialize(SrcGeometryAdapter, Accuracy);
+	}
 
-		return true;
+	bool Initialize(const TArray<FInstancedMeshMergeData>& Geometry, float Accuracy)
+	{
+		FMeshDescriptionArrayAdapter SrcGeometryAdapter(Geometry);
+		return Initialize(SrcGeometryAdapter, Accuracy);
 	}
 
 	virtual double GetVoxelSize() const override
@@ -128,6 +119,26 @@ public:
 	}
 
 private:
+	bool Initialize(FMeshDescriptionArrayAdapter& InSrcGeometryAdapter, float Accuracy)
+	{
+		OpenVDBTransform::Ptr XForm = OpenVDBTransform::createLinearTransform(Accuracy);
+		InSrcGeometryAdapter.SetTransform(XForm);
+
+		VoxelSize = InSrcGeometryAdapter.GetTransform().voxelSize()[0];
+
+		SrcPolyIndexGrid = openvdb::Int32Grid::create();
+
+		if (!ProxyLOD::MeshArrayToSDFVolume(InSrcGeometryAdapter, SDFVolume, SrcPolyIndexGrid.get()))
+		{
+			SrcPolyIndexGrid.reset();
+			return false;
+		}
+
+		Sampler.reset(new openvdb::tools::GridSampler<openvdb::FloatGrid, openvdb::tools::PointSampler>(*SDFVolume));
+
+		return true;
+	}
+
 	openvdb::FloatGrid::Ptr SDFVolume;
 	openvdb::Int32Grid::Ptr SrcPolyIndexGrid;
 	openvdb::tools::GridSampler<openvdb::FloatGrid, openvdb::tools::PointSampler>::Ptr Sampler;
@@ -151,6 +162,17 @@ TUniquePtr<IProxyLODVolume> IProxyLODVolume::CreateSDFVolumeFromMeshArray(const
 	return Volume;
 }
 
+TUniquePtr<IProxyLODVolume> IProxyLODVolume::CreateSDFVolumeFromMeshArray(const TArray<FInstancedMeshMergeData>& Geometry, float Step)
+{
+	TUniquePtr<FProxyLODVolumeImpl> Volume = MakeUnique<FProxyLODVolumeImpl>();
+
+	if (Volume == nullptr || !Volume->Initialize(Geometry, Step))
+	{
+		return nullptr;
+	}
+
+	return Volume;
+}
 
 class FPolygonSoup
 {
diff --git a/Engine/Plugins/Experimental/ProxyLODPlugin/Source/ProxyLOD/Public/ProxyLODVolume.h b/Engine/Plugins/Experimental/ProxyLODPlugin/Source/ProxyLOD/Public/ProxyLODVolume.h
--- a/Engine/Plugins/Experimental/ProxyLODPlugin/Source/ProxyLOD/Public/ProxyLODVolume.h
+++ b/Engine/Plugins/Experimental/ProxyLODPlugin/Source/ProxyLOD/Public/ProxyLODVolume.h
@@ -5,6 +5,7 @@
 #include "CoreMinimal.h"
 
 struct FMeshMergeData;
+struct FInstancedMeshMergeData;
 struct FMeshDescription;
 
 class PROXYLODMESHREDUCTION_API IProxyLODVolume
@@ -32,6 +33,7 @@ public:
 
 	/** Create OpenVDB volume from input geometry */
 	static TUniquePtr<IProxyLODVolume> CreateSDFVolumeFromMeshArray(const TArray<FMeshMergeData>& Geometry, float Step);
+	static TUniquePtr<IProxyLODVolume> CreateSDFVolumeFromMeshArray(const TArray<FInstancedMeshMergeData>& Geometry, float Step);
 
 	virtual ~IProxyLODVolume() {}
 
diff --git a/Engine/Source/Developer/MeshMergeUtilities/Private/MeshMergeUtilities.cpp b/Engine/Source/Developer/MeshMergeUtilities/Private/MeshMergeUtilities.cpp
--- a/Engine/Source/Developer/MeshMergeUtilities/Private/MeshMergeUtilities.cpp
+++ b/Engine/Source/Developer/MeshMergeUtilities/Private/MeshMergeUtilities.cpp
@@ -1208,7 +1208,13 @@ static void ScaleTextureCoordinatesToBox(const FBox2D& Box, TArray<FVector2D>& I
 
 typedef TFunctionRef<int32(const UStaticMeshComponent*)> FGetMeshLODFunc;
 
-static TArray<FMeshDescription*> GatherGeometry(const TArray<UStaticMeshComponent*>& InStaticMeshComponents, const FMeshProxySettings& InMeshProxySettings, TArray<FProxyMeshDescriptor>& InDescriptors, const TArray<TArray<int32>>& InMeshesToMergePerDescriptor, const TArray<int32>& InMeshesToMeshDescriptor, FGetMeshLODFunc InGetMeshLODFunc, int32& OutSummedLightmapPixels)
+struct FInstancedMeshDescriptionData
+{
+	FMeshDescription* MeshDescription;
+	TArray<FTransform> InstancesTransforms;
+};
+
+static TArray<FInstancedMeshDescriptionData> GatherGeometry(const TArray<UStaticMeshComponent*>& InStaticMeshComponents, const FMeshProxySettings& InMeshProxySettings, TArray<FProxyMeshDescriptor>& InDescriptors, const TArray<TArray<int32>>& InMeshesToMergePerDescriptor, const TArray<int32>& InMeshesToMeshDescriptor, FGetMeshLODFunc InGetMeshLODFunc, int32& OutSummedLightmapPixels)
 {
 	TRACE_CPUPROFILER_EVENT_SCOPE(FMeshMergeUtilities::GatherGeometry);
 
@@ -1217,7 +1223,7 @@ static TArray<FMeshDescription*> GatherGeometry(const TArray<UStaticMeshComponen
 
 	TAtomic<int32>  SummedLightmapPixels(0);
 
-	TArray<FMeshDescription*> MeshesDescriptions;
+	TArray<FInstancedMeshDescriptionData> MeshesDescriptions;
 	MeshesDescriptions.SetNum(InDescriptors.Num());
 
 	// If grouping by identical meshes, prepare each mesh description along with their flattened UVs
@@ -1255,14 +1261,22 @@ static TArray<FMeshDescription*> GatherGeometry(const TArray<UStaticMeshComponen
 			else
 			{
 				MeshDescription = InDescriptors[InMeshesToMeshDescriptor[Index]].GetMeshDescription();
-				FStaticMeshOperations::ApplyTransform(MeshDescription, StaticMeshComponent->GetComponentTransform());
+
+				if (!InMeshProxySettings.bGroupIdenticalMeshesForBaking)
+				{
+					FStaticMeshOperations::ApplyTransform(MeshDescription, StaticMeshComponent->GetComponentTransform());
+				}
 			}
 
 			// If the component is an ISMC then we need to duplicate the vertex data
 			int32 NumInstances = 1;
 			if (const UInstancedStaticMeshComponent* InstancedStaticMeshComponent = Cast<UInstancedStaticMeshComponent>(StaticMeshComponent))
 			{
-				FMeshMergeHelpers::ExpandInstances(InstancedStaticMeshComponent, MeshDescription);
+				if (!InMeshProxySettings.bGroupIdenticalMeshesForBaking)
+				{
+					FMeshMergeHelpers::ExpandInstances(InstancedStaticMeshComponent, MeshDescription);
+				}
+
 				NumInstances = InstancedStaticMeshComponent->PerInstanceSMData.Num();
 			}
 
@@ -1282,22 +1296,60 @@ static TArray<FMeshDescription*> GatherGeometry(const TArray<UStaticMeshComponen
 			AppendSettings.bMergeUVChannels[ChannelIdx] = true;
 		}
 
-		ParallelFor(InDescriptors.Num(), [&MeshesDescriptions, &InMeshesToMergePerDescriptor, &TempDescriptionData, &AppendSettings, &InMeshProxySettings](uint32 Index)
+		ParallelFor(InDescriptors.Num(), [&MeshesDescriptions, &InMeshesToMergePerDescriptor, &InStaticMeshComponents, &TempDescriptionData, &AppendSettings, &InMeshProxySettings](uint32 Index)
 		{
 			TRACE_CPUPROFILER_EVENT_SCOPE(AppendMeshes);
 			FMeshDescription* TargetMeshDescription = new FMeshDescription();
 			FStaticMeshAttributes(*TargetMeshDescription).Register();
 
-			TArray<const FMeshDescription*> SourceMeshDescriptions;
-			SourceMeshDescriptions.Reserve(InMeshesToMergePerDescriptor[Index].Num());
-			for (int32 TempIdx : InMeshesToMergePerDescriptor[Index])
+			// When using this option, do not expand the instances, but rather send their transforms to the ProxyLOD tool
+			if (InMeshProxySettings.bGroupIdenticalMeshesForBaking)
 			{
-				SourceMeshDescriptions.Add(&TempDescriptionData[TempIdx]);
+				TArray<FTransform> InstancesTransforms;
+
+				for (int32 Idx : InMeshesToMergePerDescriptor[Index])
+				{
+					UStaticMeshComponent* StaticMeshComponent = InStaticMeshComponents[Idx];
+
+					if (InMeshProxySettings.bGroupIdenticalMeshesForBaking)
+					{
+						FTransform ComponentTransform = StaticMeshComponent->GetComponentTransform();
+						FTransform ComponentTransformInv = ComponentTransform.Inverse();
+
+						if (const UInstancedStaticMeshComponent* InstancedStaticMeshComponent = Cast<UInstancedStaticMeshComponent>(StaticMeshComponent))
+						{
+							for (const FInstancedStaticMeshInstanceData& InstanceData : InstancedStaticMeshComponent->PerInstanceSMData)
+							{
+								InstancesTransforms.Add(FTransform(InstanceData.Transform) * ComponentTransform);
+							}
+						}
+						else
+						{
+							InstancesTransforms.Add(ComponentTransform);
+						}
+					}
+				}
+
+				if (!InMeshesToMergePerDescriptor[Index].IsEmpty())
+				{
+					*TargetMeshDescription = TempDescriptionData[InMeshesToMergePerDescriptor[Index][0]];
+				}
+
+				MeshesDescriptions[Index].InstancesTransforms = MoveTemp(InstancesTransforms);
 			}
+			else
+			{
+				TArray<const FMeshDescription*> SourceMeshDescriptions;
+				SourceMeshDescriptions.Reserve(InMeshesToMergePerDescriptor[Index].Num());
+				for (int32 TempIdx : InMeshesToMergePerDescriptor[Index])
+				{
+					SourceMeshDescriptions.Add(&TempDescriptionData[TempIdx]);
+				}
 
-			FStaticMeshOperations::AppendMeshDescriptions(SourceMeshDescriptions, *TargetMeshDescription, AppendSettings);
+				FStaticMeshOperations::AppendMeshDescriptions(SourceMeshDescriptions, *TargetMeshDescription, AppendSettings);
+			}
 
-			MeshesDescriptions[Index] = TargetMeshDescription;
+			MeshesDescriptions[Index].MeshDescription = TargetMeshDescription;
 		}, EParallelForFlags::Unbalanced);
 	}
 
@@ -1306,11 +1358,11 @@ static TArray<FMeshDescription*> GatherGeometry(const TArray<UStaticMeshComponen
 	return MeshesDescriptions;
 }
 
-TArray<FMeshData> PrepareBakingMeshes(const struct FMeshProxySettings& InMeshProxySettings, const TArray<FProxyMeshDescriptor>& InDescriptors, TArray<FMeshDescription*> InMeshDescriptions)
+TArray<FMeshData> PrepareBakingMeshes(const struct FMeshProxySettings& InMeshProxySettings, const TArray<FProxyMeshDescriptor>& InDescriptors, TArray<FInstancedMeshDescriptionData> InMeshDescriptionData)
 {
 	TRACE_CPUPROFILER_EVENT_SCOPE(PrepareBakingMeshes)
 
-	check(InDescriptors.Num() == InMeshDescriptions.Num());
+	check(InDescriptors.Num() == InMeshDescriptionData.Num());
 
 	TArray<FMeshData> MeshData;
 	MeshData.SetNum(InDescriptors.Num());
@@ -1319,7 +1371,7 @@ TArray<FMeshData> PrepareBakingMeshes(const struct FMeshProxySettings& InMeshPro
 	for (int32 MeshIndex = 0; MeshIndex < InDescriptors.Num(); MeshIndex++)
 	{
 		FMeshData& MeshSettings = MeshData[MeshIndex];
-		FMeshDescription& MeshDescription = *InMeshDescriptions[MeshIndex];
+		FMeshDescription& MeshDescription = *InMeshDescriptionData[MeshIndex].MeshDescription;
 		const FProxyMeshDescriptor& MeshDescriptor = InDescriptors[MeshIndex];
 
 		if (!InMeshProxySettings.bGroupIdenticalMeshesForBaking)
@@ -1333,7 +1385,7 @@ TArray<FMeshData> PrepareBakingMeshes(const struct FMeshProxySettings& InMeshPro
 	}
 
 	// Parallel step
-	ParallelFor(InDescriptors.Num(), [&MeshData, &InDescriptors, &InMeshDescriptions, &InMeshProxySettings](uint32 MeshIndex)
+	ParallelFor(InDescriptors.Num(), [&MeshData, &InDescriptors, &InMeshDescriptionData, &InMeshProxySettings](uint32 MeshIndex)
 	{
 		const FProxyMeshDescriptor& MeshDescriptor = InDescriptors[MeshIndex];
 
@@ -1349,7 +1401,7 @@ TArray<FMeshData> PrepareBakingMeshes(const struct FMeshProxySettings& InMeshPro
 		}
 		else
 		{
-			FMeshDescription& MeshDescription = *InMeshDescriptions[MeshIndex];
+			FMeshDescription& MeshDescription = *InMeshDescriptionData[MeshIndex].MeshDescription;
 			MeshSettings.MeshDescription = &MeshDescription;
 
 			TVertexInstanceAttributesRef<FVector2D> VertexInstanceUVs = FStaticMeshAttributes(MeshDescription).GetVertexInstanceUVs();
@@ -1579,7 +1631,7 @@ void FMeshMergeUtilities::CreateProxyMesh(const TArray<UStaticMeshComponent*>& I
 
 	int32 SummedLightmapPixels;
 	
-	TArray<FMeshDescription*> MeshDescriptionData = GatherGeometry(StaticMeshComponents, InMeshProxySettings, MeshDescriptors, MeshesToMergePerDescriptor, MeshToMeshDescriptor, SelectLODFunc, SummedLightmapPixels);
+	TArray<FInstancedMeshDescriptionData> MeshDescriptionData = GatherGeometry(StaticMeshComponents, InMeshProxySettings, MeshDescriptors, MeshesToMergePerDescriptor, MeshToMeshDescriptor, SelectLODFunc, SummedLightmapPixels);
 	TArray<FMeshData> MeshBakingData = PrepareBakingMeshes(InMeshProxySettings, MeshDescriptors, MeshDescriptionData);
 
 	TArray<UMaterialInterface*> UniqueMaterials;
@@ -1747,7 +1799,7 @@ void FMeshMergeUtilities::CreateProxyMesh(const TArray<UStaticMeshComponent*>& I
 			// Now have the baked out material data, need to have a map or actually remap the raw mesh data to baked material indices
 			for (int32 MeshIndex = 0; MeshIndex < MeshDescriptionData.Num(); ++MeshIndex)
 			{
-				FMeshDescription& MeshDescription = *MeshDescriptionData[MeshIndex];
+				FMeshDescription& MeshDescription = *MeshDescriptionData[MeshIndex].MeshDescription;
 
 				TArray<TPair<uint32, uint32>> SectionAndOutputIndices;
 				OutputMaterialsMap.MultiFind(MeshIndex, SectionAndOutputIndices);
@@ -1806,16 +1858,17 @@ void FMeshMergeUtilities::CreateProxyMesh(const TArray<UStaticMeshComponent*>& I
 	// Add this proxy job to map	
 	Processor->AddProxyJob(InGuid, Data);
 
-	TArray<FMeshMergeData> MergeDataEntries;
+	TArray<FInstancedMeshMergeData> MergeDataEntries;
 	{
 		TRACE_CPUPROFILER_EVENT_SCOPE(MergeDataPreparation)
 
 		for (int32 Index = 0; Index < MeshDescriptionData.Num(); ++Index)
 		{
-			FMeshMergeData MergeData;
+			FInstancedMeshMergeData MergeData;
 			MergeData.SourceStaticMesh = MeshDescriptors[Index].GetStaticMesh();
-			MergeData.RawMesh = MeshDescriptionData[Index];
+			MergeData.RawMesh = MeshDescriptionData[Index].MeshDescription;
 			MergeData.bIsClippingMesh = false;
+			MergeData.InstanceTransforms = MeshDescriptionData[Index].InstancesTransforms;
 
 			FMeshMergeHelpers::CalculateTextureCoordinateBoundsForMesh(*MergeData.RawMesh, MergeData.TexCoordBounds);
 
@@ -1850,7 +1903,7 @@ void FMeshMergeUtilities::CreateProxyMesh(const TArray<UStaticMeshComponent*>& I
 		// Populate landscape clipping geometry
 		for (FMeshDescription* RawMesh : CullingRawMeshes)
 		{
-			FMeshMergeData ClipData;
+			FInstancedMeshMergeData ClipData;
 			ClipData.bIsClippingMesh = true;
 			ClipData.RawMesh = RawMesh;
 			MergeDataEntries.Add(ClipData);
@@ -1913,7 +1966,7 @@ void FMeshMergeUtilities::CreateProxyMesh(const TArray<UStaticMeshComponent*>& I
 		MeshDescriptionData.Num(),
 		[&MeshDescriptionData](int32 Index)
 		{
-			delete MeshDescriptionData[Index];
+			delete MeshDescriptionData[Index].MeshDescription;
 		}
 	);
 }
diff --git a/Engine/Source/Developer/MeshMergeUtilities/Public/MeshMergeData.h b/Engine/Source/Developer/MeshMergeUtilities/Public/MeshMergeData.h
--- a/Engine/Source/Developer/MeshMergeUtilities/Public/MeshMergeData.h
+++ b/Engine/Source/Developer/MeshMergeUtilities/Public/MeshMergeData.h
@@ -43,6 +43,11 @@ struct FMeshMergeData
 	FString DebugJobName;
 };
 
+struct FInstancedMeshMergeData : public FMeshMergeData
+{
+	TArray<FTransform> InstanceTransforms;
+};
+
 /** Structure for encapsulating per LOD mesh merging data */
 struct FRawMeshExt
 {
diff --git a/Engine/Source/Developer/MeshReductionInterface/Public/IMeshReductionInterfaces.h b/Engine/Source/Developer/MeshReductionInterface/Public/IMeshReductionInterfaces.h
--- a/Engine/Source/Developer/MeshReductionInterface/Public/IMeshReductionInterfaces.h
+++ b/Engine/Source/Developer/MeshReductionInterface/Public/IMeshReductionInterfaces.h
@@ -99,6 +99,10 @@ public:
 		const TArray<struct FFlattenMaterial>& InputMaterials,
 		const FGuid InJobGUID) {}
 	
+	virtual void ProxyLOD(const TArray<struct FInstancedMeshMergeData>& InData,
+		const struct FMeshProxySettings& InProxySettings,
+		const TArray<struct FFlattenMaterial>& InputMaterials,
+		const FGuid InJobGUID) {}
 
 	virtual void AggregateLOD() {}
 
