Fixed a shadowing variable error in clang on Windows.

#trivial
#rnx

diff --git a/Engine/Plugins/Experimental/ProxyLODPlugin/Source/ProxyLOD/Private/ProxyLODVolume.cpp b/Engine/Plugins/Experimental/ProxyLODPlugin/Source/ProxyLOD/Private/ProxyLODVolume.cpp
--- a/Engine/Plugins/Experimental/ProxyLODPlugin/Source/ProxyLOD/Private/ProxyLODVolume.cpp
+++ b/Engine/Plugins/Experimental/ProxyLODPlugin/Source/ProxyLOD/Private/ProxyLODVolume.cpp
@@ -466,9 +466,9 @@ private:
 
 		auto TransformGenerator = [&LocalToVoxel, &AverageTranslation](const FPlacedMesh& PlacedMesh)->openvdb::Mat4R
 		{
-			FTransform XForm = PlacedMesh.Transform;
-			XForm.AddToTranslation(-AverageTranslation);
-			FMatrix TransformMatrix = XForm.ToMatrixWithScale().Inverse();
+			FTransform MeshXForm = PlacedMesh.Transform;
+			MeshXForm.AddToTranslation(-AverageTranslation);
+			FMatrix TransformMatrix = MeshXForm.ToMatrixWithScale().Inverse();
 
 			TransformMatrix = LocalToVoxel * TransformMatrix;
 			float* data = &TransformMatrix.M[0][0];
