Fixed issues when ProxyLOD receives no instances

diff --git a/Engine/Plugins/Experimental/ProxyLODPlugin/Source/ProxyLOD/Private/ProxyLODMeshTypes.cpp b/Engine/Plugins/Experimental/ProxyLODPlugin/Source/ProxyLOD/Private/ProxyLODMeshTypes.cpp
--- a/Engine/Plugins/Experimental/ProxyLODPlugin/Source/ProxyLOD/Private/ProxyLODMeshTypes.cpp
+++ b/Engine/Plugins/Experimental/ProxyLODPlugin/Source/ProxyLOD/Private/ProxyLODMeshTypes.cpp
@@ -159,7 +159,7 @@ FMeshDescriptionArrayAdapter::FMeshDescriptionArrayAdapter(const TArray<const FI
 	Transform = openvdb::math::Transform::createLinearTransform(1.);
 
 	SetupInstances(InMergeDataPtrArray.Num(), [&InMergeDataPtrArray](uint32 Index) { return InMergeDataPtrArray[Index]; });
-	Construct(InMergeDataPtrArray.Num(), [&InMergeDataPtrArray](uint32 Index) { return InMergeDataPtrArray[Index]; }, [&InMergeDataPtrArray](uint32 Index) { return InMergeDataPtrArray[Index]->InstanceTransforms.Num(); });
+	Construct(InMergeDataPtrArray.Num(), [&InMergeDataPtrArray](uint32 Index) { return InMergeDataPtrArray[Index]; }, [&InMergeDataPtrArray](uint32 Index) { return InMergeDataPtrArray[Index]->InstanceTransforms.IsEmpty() ? 1 : InMergeDataPtrArray[Index]->InstanceTransforms.Num(); });
 }
 
 FMeshDescriptionArrayAdapter::FMeshDescriptionArrayAdapter(const TArray<const FMeshMergeData*>& InMergeDataPtrArray)
@@ -404,7 +404,7 @@ const FMeshDescription& FMeshDescriptionArrayAdapter::GetRawMesh(const size_t Fa
 	LocalFaceNumber = int32(FaceNumber) - PolyOffsetArray[MeshIdx];
 
 	int32 InstanceCount = !InstancesTransformArray.IsEmpty() ? InstancesTransformArray[MeshIdx].Num() : INDEX_NONE;
-	if (InstanceCount != INDEX_NONE)
+	if (InstanceCount > 0)
 	{
 		int32 NumFacesPerInstance = RawMeshArrayData[MeshIdx].TriangleCount;
 		InstanceIdx = LocalFaceNumber / NumFacesPerInstance;
