#ue5 - LWC: fix compile error on Clang. Initializer lists seem to ignore compiler truncation warning flags, so make truncation explicit here.

diff --git a/Engine/Plugins/Experimental/ProxyLODPlugin/Source/ProxyLOD/Private/ProxyLODMaterialTransferUtilities.cpp b/Engine/Plugins/Experimental/ProxyLODPlugin/Source/ProxyLOD/Private/ProxyLODMaterialTransferUtilities.cpp
--- a/Engine/Plugins/Experimental/ProxyLODPlugin/Source/ProxyLOD/Private/ProxyLODMaterialTransferUtilities.cpp
+++ b/Engine/Plugins/Experimental/ProxyLODPlugin/Source/ProxyLOD/Private/ProxyLODMaterialTransferUtilities.cpp
@@ -690,7 +690,7 @@ namespace
 							const FVector2D SrcTexelLocation(SrcUV.X * SrcBufferSize.X, SrcUV.Y * SrcBufferSize.Y);
 
 							// 0,0 corner
-							const FIntPoint Min(FMath::FloorToInt(SrcTexelLocation.X), FMath::FloorToInt(SrcTexelLocation.Y));
+							const FIntPoint Min(FMath::FloorToInt32(SrcTexelLocation.X), FMath::FloorToInt32(SrcTexelLocation.Y));
 
 							// offset.
 							const FVector2D Delta(SrcTexelLocation.X - (float)Min.X, SrcTexelLocation.Y - (float)Min.Y);
@@ -786,8 +786,8 @@ namespace
 
 					// Where are these corners in the SrcGrid 
 					// in the grid index space (IS) ?
-					int32 SrcISMin[2] = { FMath::FloorToInt(DstUSMin[0] * SrcI), FMath::FloorToInt(DstUSMin[1] * SrcJ) };
-					int32 SrcISMax[2] = { FMath::CeilToInt(DstUSMax[0] * SrcI), FMath::CeilToInt(DstUSMax[1] * SrcJ) };
+					int32 SrcISMin[2] = { FMath::FloorToInt32(DstUSMin[0] * SrcI), FMath::FloorToInt32(DstUSMin[1] * SrcJ) };
+					int32 SrcISMax[2] = { FMath::CeilToInt32(DstUSMax[0] * SrcI), FMath::CeilToInt32(DstUSMax[1] * SrcJ) };
 
 					// clip against boundary
 					SrcISMin[0] = FMath::Max(SrcISMin[0], 0);
@@ -1178,7 +1178,7 @@ namespace
 
 							// 0,0 corner
 
-							const FIntPoint Min(FMath::FloorToInt(SrcTexelLocation.X), FMath::FloorToInt(SrcTexelLocation.Y));
+							const FIntPoint Min(FMath::FloorToInt32(SrcTexelLocation.X), FMath::FloorToInt32(SrcTexelLocation.Y));
 
 							// offset.
 
