Mesh Modeling - update the voxel-based CSG and Merge tools to have the same output selection and widgets that the other voxel-based tools have - this was done by converting the tools to use the UBaseVoxelTool base class.
Also hook-up progress class /  interrputer to that open-vdb voxelization step so a cancelled tool may use less resources.

diff --git a/Engine/Plugins/Experimental/MeshModelingToolsetExp/Source/MeshModelingToolsEditorOnlyExp/Private/MergeMeshesTool.cpp b/Engine/Plugins/Experimental/MeshModelingToolsetExp/Source/MeshModelingToolsEditorOnlyExp/Private/MergeMeshesTool.cpp
--- a/Engine/Plugins/Experimental/MeshModelingToolsetExp/Source/MeshModelingToolsEditorOnlyExp/Private/MergeMeshesTool.cpp
+++ b/Engine/Plugins/Experimental/MeshModelingToolsetExp/Source/MeshModelingToolsEditorOnlyExp/Private/MergeMeshesTool.cpp
@@ -3,228 +3,98 @@
 #include "MergeMeshesTool.h"
 #include "InteractiveToolManager.h"
 #include "ToolBuilderUtil.h"
+
 #include "ToolSetupUtil.h"
-#include "ModelingObjectsCreationAPI.h"
-#include "ModelingToolTargetUtil.h"
+
 #include "Selection/ToolSelectionUtil.h"
 
 #include "DynamicMesh/DynamicMesh3.h"
-#include "MeshDescriptionToDynamicMesh.h"
-
-#include "CompositionOps/VoxelMergeMeshesOp.h"
-
-#include "TargetInterfaces/MeshDescriptionProvider.h"
-#include "TargetInterfaces/PrimitiveComponentBackedTarget.h"
-#include "ToolTargetManager.h"
-
-#include UE_INLINE_GENERATED_CPP_BY_NAME(MergeMeshesTool)
-
-using namespace UE::Geometry;
+#include "DynamicMesh/MeshTransforms.h"
 
-#define LOCTEXT_NAMESPACE "UMergeMeshesTool"
 
-/*
- * ToolBuilder
- */
 
+#include "InteractiveGizmoManager.h"
 
-const FToolTargetTypeRequirements& UMergeMeshesToolBuilder::GetTargetRequirements() const
-{
-	static FToolTargetTypeRequirements TypeRequirements({
-		UMeshDescriptionProvider::StaticClass(),
-		UPrimitiveComponentBackedTarget::StaticClass()
-		});
-	return TypeRequirements;
-}
+#include "BaseGizmos/GizmoComponents.h"
+#include "BaseGizmos/CombinedTransformGizmo.h"
 
-bool UMergeMeshesToolBuilder::CanBuildTool(const FToolBuilderState& SceneState) const
-{
-	constexpr int32 MinRequiredComponents = 1;
-	const bool bHasComponents = SceneState.TargetManager->CountSelectedAndTargetable(SceneState, GetTargetRequirements()) >= MinRequiredComponents;
-	return ( bHasComponents );
-}
-
-UMultiSelectionMeshEditingTool* UMergeMeshesToolBuilder::CreateNewTool(const FToolBuilderState& SceneState) const
-{
-	return NewObject<UMergeMeshesTool>(SceneState.ToolManager);
-}
+#include "CompositionOps/VoxelMergeMeshesOp.h"
 
+#include UE_INLINE_GENERATED_CPP_BY_NAME(MergeMeshesTool)
 
+using namespace UE::Geometry;
 
+#define LOCTEXT_NAMESPACE "UMergeMeshesTool"
 
-/*
- * Tool
- */
 
-UMergeMeshesTool::UMergeMeshesTool()
-{
-}
 
 
-void UMergeMeshesTool::Setup()
+void UMergeMeshesTool::SetupProperties()
 {
-	UInteractiveTool::Setup();
+	//Super::SetupProperties();
+	//RemoveToolPropertySource(VoxProperties);
 
-	MergeProps = NewObject<UMergeMeshesToolProperties>();
+	UBaseCreateFromSelectedTool::SetupProperties();
+	MergeProps = NewObject<UMergeMeshesToolProperties>(this);
 	MergeProps->RestoreProperties(this);
 	AddToolPropertySource(MergeProps);
 
-	MeshStatisticsProperties = NewObject<UMeshStatisticsProperties>(this);
-	AddToolPropertySource(MeshStatisticsProperties);
-
-	HandleSourcesProperties = NewObject<UOnAcceptHandleSourcesProperties>(this);
-	HandleSourcesProperties->RestoreProperties(this);
-	AddToolPropertySource(HandleSourcesProperties);
-
-	// Hide the source meshes
-	for (int32 ComponentIdx = 0; ComponentIdx < Targets.Num(); ComponentIdx++)
-	{
-		UE::ToolTarget::HideSourceObject(Targets[ComponentIdx]);
-	}
-
-	// save transformed version of input meshes (maybe this could happen in the Operator?)
-	CacheInputMeshes();
-
-	// initialize the PreviewMesh+BackgroundCompute object
-	Preview = NewObject<UMeshOpPreviewWithBackgroundCompute>(this, "Preview");
-	Preview->Setup(GetTargetWorld(), this);
-	Preview->OnMeshUpdated.AddLambda([this](UMeshOpPreviewWithBackgroundCompute* Compute) {
-		MeshStatisticsProperties->Update(*Compute->PreviewMesh->GetPreviewDynamicMesh());
-		UpdateAcceptWarnings(Compute->HaveEmptyResult() ? EAcceptWarning::EmptyForbidden : EAcceptWarning::NoWarning);
-	});
-	ToolSetupUtil::ApplyRenderingConfigurationToPreview(Preview->PreviewMesh, nullptr);
-
-	CreateLowQualityPreview(); // update the preview with a low-quality result
-	
-	Preview->ConfigureMaterials(
-		ToolSetupUtil::GetDefaultSculptMaterial(GetToolManager()),
-		ToolSetupUtil::GetDefaultWorkingMaterial(GetToolManager())
-	);
-
-	Preview->InvalidateResult();    // start compute
-
-
-	SetToolDisplayName(LOCTEXT("ToolName", "Voxel Merge"));
+	SetToolDisplayName(LOCTEXT("UMergeMeshesToolName", "Voxel Merge Meshes"));
 	GetToolManager()->DisplayMessage(
-		LOCTEXT("OnStartTool", "Combine the input meshes into closed solids using voxelization techniques. UVs, sharp edges, and small/thin features will be lost. Increase Voxel Count to enhance accuracy."),
+		LOCTEXT("OnStartTool", "Compute a CSG Union of the input meshes using voxelization techniques.UVs, sharp edges, and small / thin features will be lost.Increase Voxel Count to enhance accuracy."),
 		EToolMessageLevel::UserNotification);
-}
 
+}
 
-void UMergeMeshesTool::OnShutdown(EToolShutdownType ShutdownType)
+void UMergeMeshesTool::SaveProperties()
 {
+	UBaseCreateFromSelectedTool::SaveProperties();
+	//Super::SaveProperties();
 	MergeProps->SaveProperties(this);
-	HandleSourcesProperties->SaveProperties(this);
-
-	FDynamicMeshOpResult Result = Preview->Shutdown();
-	// Restore (unhide) the source meshes
-	for (int32 ComponentIdx = 0; ComponentIdx < Targets.Num(); ComponentIdx++)
-	{
-		UE::ToolTarget::ShowSourceObject(Targets[ComponentIdx]);
-	}
-	if (ShutdownType == EToolShutdownType::Accept)
-	{
-		GetToolManager()->BeginUndoTransaction(LOCTEXT("MergeMeshes", "Merge Meshes"));
-
-		// Generate the result
-		GenerateAsset(Result);
-
-		TArray<AActor*> Actors;
-		for (int32 ComponentIdx = 0; ComponentIdx < Targets.Num(); ComponentIdx++)
-		{
-			Actors.Add(UE::ToolTarget::GetTargetActor(Targets[ComponentIdx]));
-		}
-		HandleSourcesProperties->ApplyMethod(Actors, GetToolManager());
-
-		GetToolManager()->EndUndoTransaction();
-	}
 }
 
 
-
-void UMergeMeshesTool::OnTick(float DeltaTime)
+void UMergeMeshesTool::ConvertInputsAndSetPreviewMaterials(bool bSetPreviewMesh)
 {
-	Preview->Tick(DeltaTime);
-}
-
-bool UMergeMeshesTool::CanAccept() const
-{
-	return Super::CanAccept() && Preview->HaveValidNonEmptyResult();
-}
+	Super::ConvertInputsAndSetPreviewMaterials(bSetPreviewMesh);
 
-void UMergeMeshesTool::OnPropertyModified(UObject* PropertySet, FProperty* Property)
-{
-	Preview->InvalidateResult();
+	if (HasOpenBoundariesInMeshInputs())
+	{
+		GetToolManager()->DisplayMessage(
+			LOCTEXT("WarnOpenEdges", "Open edges found: some artifacts will result."),
+			EToolMessageLevel::UserWarning);
+	}
 }
 
-
-
-
 TUniquePtr<FDynamicMeshOperator> UMergeMeshesTool::MakeNewOperator()
 {
 	TUniquePtr<FVoxelMergeMeshesOp> MergeOp = MakeUnique<FVoxelMergeMeshesOp>();
-	MergeOp->VoxelCount     = MergeProps->VoxelCount;
-	MergeOp->AdaptivityD    = MergeProps->MeshAdaptivity;
-	MergeOp->IsoSurfaceD    = MergeProps->OffsetDistance;
-	MergeOp->bAutoSimplify  = MergeProps->bAutoSimplify;
-	MergeOp->InputMeshArray = InputMeshes;
+	MergeOp->VoxelCount = MergeProps->VoxelCount;
+	MergeOp->AdaptivityD = MergeProps->MeshAdaptivity;
+	MergeOp->IsoSurfaceD = MergeProps->OffsetDistance;
+	MergeOp->bAutoSimplify = MergeProps->bAutoSimplify;
+	MergeOp->Transforms.SetNum(Targets.Num());
+	MergeOp->Meshes.SetNum(Targets.Num());
+	for (int Idx = 0; Idx < Targets.Num(); Idx++)
+	{
+		MergeOp->Meshes[Idx] = OriginalDynamicMeshes[Idx];
+		MergeOp->Transforms[Idx] = TransformProxies[Idx]->GetTransform();
+	}
 	return MergeOp;
 }
 
 
-
-void UMergeMeshesTool::CacheInputMeshes()
+FString UMergeMeshesTool::GetCreatedAssetName() const
 {
-	InputMeshes.Reset();
-
-	// Package the selected meshes and transforms for consumption by the CSGTool
-	for (int32 ComponentIdx = 0; ComponentIdx < Targets.Num(); ComponentIdx++)
-	{
-		UE::Geometry::FVoxelMergeMeshesOp::FInputMesh InputMesh;
-		InputMesh.Mesh = UE::ToolTarget::GetMeshDescription(Targets[ComponentIdx]);
-		InputMesh.Transform = (FTransform) UE::ToolTarget::GetLocalToWorldTransform(Targets[ComponentIdx]);
-		InputMeshes.Add(InputMesh);
-	}
+	return TEXT("VoxelMergedMeshes");
 }
 
-void UMergeMeshesTool::CreateLowQualityPreview()
+FText UMergeMeshesTool::GetActionName() const
 {
-
-	FProgressCancel NullInterrupter;
-	FVoxelMergeMeshesOp MergeMeshesOp;
-
-	MergeMeshesOp.VoxelCount = 12;
-	MergeMeshesOp.AdaptivityD = 0.001;
-	MergeMeshesOp.bAutoSimplify = true;
-	MergeMeshesOp.InputMeshArray = InputMeshes;
-	
-	MergeMeshesOp.CalculateResult(&NullInterrupter);
-	TUniquePtr<FDynamicMesh3> FastPreviewMesh = MergeMeshesOp.ExtractResult();
-
-
-	Preview->PreviewMesh->SetTransform((FTransform)MergeMeshesOp.GetResultTransform());
-	Preview->PreviewMesh->UpdatePreview(FastPreviewMesh.Get());  // copies the mesh @todo we could just give ownership to the Preview!
-	Preview->SetVisibility(true);
+	return LOCTEXT("VoxelMergeMeshes", "Voxel Merge Meshes");
 }
 
 
-void UMergeMeshesTool::GenerateAsset(const FDynamicMeshOpResult& OpResult)
-{
-	check(OpResult.Mesh.Get() != nullptr);
-	
-	FCreateMeshObjectParams NewMeshObjectParams;
-	NewMeshObjectParams.TargetWorld = GetTargetWorld();
-	NewMeshObjectParams.Transform = (FTransform)OpResult.Transform;
-	NewMeshObjectParams.BaseName = TEXT("MergedMesh");
-	NewMeshObjectParams.Materials.Add(ToolSetupUtil::GetDefaultMaterial());
-	NewMeshObjectParams.SetMesh(OpResult.Mesh.Get());
-	FCreateMeshObjectResult Result = UE::Modeling::CreateMeshObject(GetToolManager(), MoveTemp(NewMeshObjectParams));
-	if (Result.IsOK() && Result.NewActor != nullptr)
-	{
-		ToolSelectionUtil::SetNewActorSelection(GetToolManager(), Result.NewActor);
-	}
-
-}
 
 
 
diff --git a/Engine/Plugins/Experimental/MeshModelingToolsetExp/Source/MeshModelingToolsEditorOnlyExp/Private/VoxelCSGMeshesTool.cpp b/Engine/Plugins/Experimental/MeshModelingToolsetExp/Source/MeshModelingToolsEditorOnlyExp/Private/VoxelCSGMeshesTool.cpp
--- a/Engine/Plugins/Experimental/MeshModelingToolsetExp/Source/MeshModelingToolsEditorOnlyExp/Private/VoxelCSGMeshesTool.cpp
+++ b/Engine/Plugins/Experimental/MeshModelingToolsetExp/Source/MeshModelingToolsEditorOnlyExp/Private/VoxelCSGMeshesTool.cpp
@@ -1,22 +1,25 @@
 // Copyright Epic Games, Inc. All Rights Reserved.
 
 #include "VoxelCSGMeshesTool.h"
+
 #include "InteractiveToolManager.h"
 #include "ToolBuilderUtil.h"
+
 #include "ToolSetupUtil.h"
-#include "ModelingObjectsCreationAPI.h"
+
 #include "Selection/ToolSelectionUtil.h"
-#include "ModelingToolTargetUtil.h"
 
 #include "DynamicMesh/DynamicMesh3.h"
-#include "MeshDescriptionToDynamicMesh.h"
+#include "DynamicMesh/MeshTransforms.h"
 
-#include "CompositionOps/VoxelBooleanMeshesOp.h"
 
-#include "TargetInterfaces/MeshDescriptionCommitter.h"
-#include "TargetInterfaces/MeshDescriptionProvider.h"
-#include "TargetInterfaces/PrimitiveComponentBackedTarget.h"
-#include "ToolTargetManager.h"
+
+#include "InteractiveGizmoManager.h"
+
+#include "BaseGizmos/GizmoComponents.h"
+#include "BaseGizmos/CombinedTransformGizmo.h"
+
+#include "CompositionOps/VoxelBooleanMeshesOp.h"
 
 #include UE_INLINE_GENERATED_CPP_BY_NAME(VoxelCSGMeshesTool)
 
@@ -25,191 +28,76 @@ using namespace UE::Geometry;
 #define LOCTEXT_NAMESPACE "UVoxelCSGMeshesTool"
 
 
-/*
- * ToolBuilder
- */
-bool UVoxelCSGMeshesToolBuilder::CanBuildTool(const FToolBuilderState& SceneState) const
+void UVoxelCSGMeshesTool::SetupProperties()
 {
-	return SceneState.TargetManager->CountSelectedAndTargetable(SceneState, GetTargetRequirements()) == 2;
-}
-
-UMultiSelectionMeshEditingTool* UVoxelCSGMeshesToolBuilder::CreateNewTool(const FToolBuilderState& SceneState) const
-{
-	return NewObject<UVoxelCSGMeshesTool>(SceneState.ToolManager);
-}
+	//Super::SetupProperties();
+	//RemoveToolPropertySource(VoxProperties);
 
-/*
- * Tool
- */
+	UBaseCreateFromSelectedTool::SetupProperties();
 
-UVoxelCSGMeshesTool::UVoxelCSGMeshesTool()
-{
-}
-
-void UVoxelCSGMeshesTool::Setup()
-{
-	UInteractiveTool::Setup();
-
-	CSGProps = NewObject<UVoxelCSGMeshesToolProperties>();
+	CSGProps = NewObject<UVoxelCSGMeshesToolProperties>(this);
 	CSGProps->RestoreProperties(this);
 	AddToolPropertySource(CSGProps);
 
-	MeshStatisticsProperties = NewObject<UMeshStatisticsProperties>(this);
-	AddToolPropertySource(MeshStatisticsProperties);
-
-	HandleSourcesProperties = NewObject<UOnAcceptHandleSourcesProperties>(this);
-	HandleSourcesProperties->RestoreProperties(this);
-	AddToolPropertySource(HandleSourcesProperties);
-
-
-	// Hide the source meshes
-	for (int32 ComponentIdx = 0; ComponentIdx < Targets.Num(); ComponentIdx++)
-	{
-		UE::ToolTarget::HideSourceObject(Targets[ComponentIdx]);
-	}
-
-	// save transformed version of input meshes (maybe this could happen in the Operator?)
-	CacheInputMeshes();
-
-	// initialize the PreviewMesh+BackgroundCompute object
-	Preview = NewObject<UMeshOpPreviewWithBackgroundCompute>(this, "Preview");
-	Preview->Setup(GetTargetWorld(), this);
-	ToolSetupUtil::ApplyRenderingConfigurationToPreview(Preview->PreviewMesh, nullptr);
-	Preview->OnMeshUpdated.AddLambda([this](UMeshOpPreviewWithBackgroundCompute* Compute) {
-		MeshStatisticsProperties->Update(*Compute->PreviewMesh->GetPreviewDynamicMesh());
-		UpdateAcceptWarnings(Compute->HaveEmptyResult() ? EAcceptWarning::EmptyForbidden : EAcceptWarning::NoWarning);
-	});
-
-	CreateLowQualityPreview();
-
-	Preview->ConfigureMaterials(
-		ToolSetupUtil::GetDefaultSculptMaterial(GetToolManager()),
-		ToolSetupUtil::GetDefaultWorkingMaterial(GetToolManager())
-	);
-	
-	Preview->InvalidateResult();    // start compute
-
-	SetToolDisplayName(LOCTEXT("ToolName", "Voxel Boolean"));
+	SetToolDisplayName(LOCTEXT("UVoxelCSGMeshesToolName", "Voxel Boolean"));
 	GetToolManager()->DisplayMessage(
-		LOCTEXT("OnStartTool", "Compute a CSG Boolean of the input meshes using voxelization techniques. UVs, sharp edges, and small/thin features will be lost. Increase Voxel Count to enhance accuracy."),
+		LOCTEXT("OnStartTool", "Compute a CSG Boolean of the input meshes using voxelization techniques.UVs, sharp edges, and small / thin features will be lost.Increase Voxel Count to enhance accuracy."),
 		EToolMessageLevel::UserNotification);
-}
 
-
-void UVoxelCSGMeshesTool::OnShutdown(EToolShutdownType ShutdownType)
-{
-	CSGProps->SaveProperties(this);
-	HandleSourcesProperties->SaveProperties(this);
-
-	FDynamicMeshOpResult Result = Preview->Shutdown();
-	// Restore (unhide) the source meshes
-	for (int32 ComponentIdx = 0; ComponentIdx < Targets.Num(); ComponentIdx++)
-	{
-		UE::ToolTarget::ShowSourceObject(Targets[ComponentIdx]);
-	}
-	if (ShutdownType == EToolShutdownType::Accept)
-	{
-		GetToolManager()->BeginUndoTransaction(LOCTEXT("BooleanMeshes", "Boolean Meshes"));
-
-		// Generate the result
-		GenerateAsset(Result);
-
-		TArray<AActor*> Actors;
-		for (int32 ComponentIdx = 0; ComponentIdx < Targets.Num(); ComponentIdx++)
-		{
-			Actors.Add(UE::ToolTarget::GetTargetActor(Targets[ComponentIdx]));
-		}
-		HandleSourcesProperties->ApplyMethod(Actors, GetToolManager());
-
-		GetToolManager()->EndUndoTransaction();
-	}
 }
 
 
-
-void UVoxelCSGMeshesTool::OnTick(float DeltaTime)
+void UVoxelCSGMeshesTool::SaveProperties()
 {
-	Preview->Tick(DeltaTime);
+	//Super::SaveProperties();
+	UBaseCreateFromSelectedTool::SaveProperties();
+	CSGProps->SaveProperties(this);
 }
 
-void UVoxelCSGMeshesTool::Render(IToolsContextRenderAPI* RenderAPI)
-{
-}
 
-bool UVoxelCSGMeshesTool::CanAccept() const
+void UVoxelCSGMeshesTool::ConvertInputsAndSetPreviewMaterials(bool bSetPreviewMesh)
 {
-	return Super::CanAccept() && Preview->HaveValidNonEmptyResult();
-}
+	Super::ConvertInputsAndSetPreviewMaterials(bSetPreviewMesh);
 
-void UVoxelCSGMeshesTool::OnPropertyModified(UObject* PropertySet, FProperty* Property)
-{
-	Preview->InvalidateResult();
+	if (HasOpenBoundariesInMeshInputs())
+	{
+		GetToolManager()->DisplayMessage(
+			LOCTEXT("WarnOpenEdges", "Open edges found: some artifacts will result."),
+			EToolMessageLevel::UserWarning);
+	}
 }
 
 TUniquePtr<FDynamicMeshOperator> UVoxelCSGMeshesTool::MakeNewOperator()
 {
 	TUniquePtr<FVoxelBooleanMeshesOp> CSGOp = MakeUnique<FVoxelBooleanMeshesOp>();
-	CSGOp->Operation      = (FVoxelBooleanMeshesOp::EBooleanOperation)(int)CSGProps->Operation;
-	CSGOp->VoxelCount     = CSGProps->VoxelCount;
-	CSGOp->AdaptivityD    = CSGProps->MeshAdaptivity;
-	CSGOp->IsoSurfaceD    = CSGProps->OffsetDistance;
-	CSGOp->bAutoSimplify  = CSGProps->bAutoSimplify;
-	CSGOp->InputMeshArray = InputMeshes;
-	return CSGOp;
-}
-
-void UVoxelCSGMeshesTool::CacheInputMeshes()
-{
-	InputMeshes.Reset();
-
-	// Package the selected meshes and transforms for consumption by the CSGTool
-	for (int32 ComponentIdx = 0; ComponentIdx < Targets.Num(); ComponentIdx++)
+	CSGOp->Operation = (FVoxelBooleanMeshesOp::EBooleanOperation)(int)CSGProps->Operation;
+	CSGOp->VoxelCount = CSGProps->VoxelCount;
+	CSGOp->AdaptivityD = CSGProps->MeshAdaptivity;
+	CSGOp->IsoSurfaceD = CSGProps->OffsetDistance;
+	CSGOp->bAutoSimplify = CSGProps->bAutoSimplify;
+	CSGOp->Transforms.SetNum(Targets.Num());
+	CSGOp->Meshes.SetNum(Targets.Num());
+	for (int Idx = 0; Idx < Targets.Num(); Idx++)
 	{
-		UE::Geometry::FVoxelBooleanMeshesOp::FInputMesh InputMesh;
-		InputMesh.Mesh = UE::ToolTarget::GetMeshDescription(Targets[ComponentIdx]);
-		InputMesh.Transform = (FTransform) UE::ToolTarget::GetLocalToWorldTransform(Targets[ComponentIdx]);
-		InputMeshes.Add(InputMesh);
+		CSGOp->Meshes[Idx] = OriginalDynamicMeshes[Idx];
+		CSGOp->Transforms[Idx] = TransformProxies[Idx]->GetTransform();
 	}
+	return CSGOp;
 }
 
-void UVoxelCSGMeshesTool::CreateLowQualityPreview()
-{
 
-	FProgressCancel NullInterrupter;
-	FVoxelBooleanMeshesOp BooleanOp;
 
-	BooleanOp.Operation = (FVoxelBooleanMeshesOp::EBooleanOperation)(int)CSGProps->Operation;
-	BooleanOp.VoxelCount = 12;
-	BooleanOp.AdaptivityD = 0.01;
-	BooleanOp.bAutoSimplify = true;
-	BooleanOp.InputMeshArray = InputMeshes;
-	
-	BooleanOp.CalculateResult(&NullInterrupter);
-	TUniquePtr<FDynamicMesh3> FastPreviewMesh = BooleanOp.ExtractResult();
-
-
-	Preview->PreviewMesh->SetTransform((FTransform)BooleanOp.GetResultTransform());
-	Preview->PreviewMesh->UpdatePreview(FastPreviewMesh.Get());  // copies the mesh @todo we could just give ownership to the Preview!
-	Preview->SetVisibility(true);
+FString UVoxelCSGMeshesTool::GetCreatedAssetName() const
+{
+	return TEXT("VoxelBoolean");
 }
 
-void UVoxelCSGMeshesTool::GenerateAsset(const FDynamicMeshOpResult& OpResult)
+FText UVoxelCSGMeshesTool::GetActionName() const
 {
-	if (ensure(OpResult.Mesh.Get() != nullptr))
-	{
-		FCreateMeshObjectParams NewMeshObjectParams;
-		NewMeshObjectParams.TargetWorld = GetTargetWorld();
-		NewMeshObjectParams.Transform = (FTransform)OpResult.Transform;
-		NewMeshObjectParams.BaseName = TEXT("CSGMesh");
-		NewMeshObjectParams.Materials.Add(ToolSetupUtil::GetDefaultMaterial());
-		NewMeshObjectParams.SetMesh(OpResult.Mesh.Get());
-		FCreateMeshObjectResult Result = UE::Modeling::CreateMeshObject(GetToolManager(), MoveTemp(NewMeshObjectParams));
-		if (Result.IsOK() && Result.NewActor != nullptr)
-		{
-			ToolSelectionUtil::SetNewActorSelection(GetToolManager(), Result.NewActor);
-		}
-	}
+	return LOCTEXT("VoxelCSGMeshes", "Voxel Boolean");
 }
 
+
+
 #undef LOCTEXT_NAMESPACE
 
diff --git a/Engine/Plugins/Experimental/MeshModelingToolsetExp/Source/MeshModelingToolsEditorOnlyExp/Public/MergeMeshesTool.h b/Engine/Plugins/Experimental/MeshModelingToolsetExp/Source/MeshModelingToolsEditorOnlyExp/Public/MergeMeshesTool.h
--- a/Engine/Plugins/Experimental/MeshModelingToolsetExp/Source/MeshModelingToolsEditorOnlyExp/Public/MergeMeshesTool.h
+++ b/Engine/Plugins/Experimental/MeshModelingToolsetExp/Source/MeshModelingToolsEditorOnlyExp/Public/MergeMeshesTool.h
@@ -3,29 +3,11 @@
 #pragma once
 
 #include "CoreMinimal.h"
-#include "BaseTools/MultiSelectionMeshEditingTool.h"
-#include "InteractiveToolBuilder.h"
-#include "MeshOpPreviewHelpers.h"
-#include "Properties/MeshStatisticsProperties.h"
-#include "PropertySets/OnAcceptProperties.h"
+#include "UObject/NoExportTypes.h"
+#include "BaseTools/BaseVoxelTool.h"
 #include "CompositionOps/VoxelMergeMeshesOp.h"
 #include "MergeMeshesTool.generated.h"
 
-/**
- *
- */
-UCLASS()
-class MESHMODELINGTOOLSEDITORONLYEXP_API UMergeMeshesToolBuilder : public UMultiSelectionMeshEditingToolBuilder
-{
-	GENERATED_BODY()
-
-public:
-	virtual bool CanBuildTool(const FToolBuilderState& SceneState) const override;
-	virtual UMultiSelectionMeshEditingTool* CreateNewTool(const FToolBuilderState& SceneState) const override;
-
-protected:
-	virtual const FToolTargetTypeRequirements& GetTargetRequirements() const override;
-};
 
 
 
@@ -39,68 +21,65 @@ class MESHMODELINGTOOLSEDITORONLYEXP_API UMergeMeshesToolProperties : public UIn
 public:
 
 	/** The size of the geometry bounding box major axis measured in voxels.*/
-	UPROPERTY(EditAnywhere, Category = Options, meta = (UIMin = "8", UIMax = "1024", ClampMin = "8", ClampMax = "1024"))
+	UPROPERTY(EditAnywhere, Category = VoxelSettings, meta = (UIMin = "8", UIMax = "1024", ClampMin = "8", ClampMax = "1024"))
 	int32 VoxelCount = 128;
 
 	/** Remeshing adaptivity, prior to optional simplification */
-	UPROPERTY(EditAnywhere, Category = Options, meta = (UIMin = "0", UIMax = "1", ClampMin = "0", ClampMax = "1"))
+	UPROPERTY(EditAnywhere, Category = VoxelSettings, meta = (UIMin = "0", UIMax = "1", ClampMin = "0", ClampMax = "1"))
 	float MeshAdaptivity = 0.001f;
 
-	/** Offset when remeshing, note large offsets with high voxels counts will be slow */
-	UPROPERTY(EditAnywhere, Category = Options, meta = (UIMin = "-10", UIMax = "10", ClampMin = "-10", ClampMax = "10"))
+	/** Offset when remeshing, note large offsets with high voxels counts will be slow.  Hidden because this duplicates functionality of the voxel offset tool */
+	UPROPERTY()
 	float OffsetDistance = 0;
 
 	/** Automatically simplify the result of voxel-based merge.*/
-	UPROPERTY(EditAnywhere, Category = Options)
+	UPROPERTY(EditAnywhere, Category = VoxelSettings)
 	bool bAutoSimplify = false;
 };
 
-/**
- *
- */
-UCLASS()
-class MESHMODELINGTOOLSEDITORONLYEXP_API UMergeMeshesTool : public UMultiSelectionMeshEditingTool, public UE::Geometry::IDynamicMeshOperatorFactory
+UCLASS()  
+class MESHMODELINGTOOLSEDITORONLYEXP_API UMergeMeshesTool : public UBaseVoxelTool
 {
 	GENERATED_BODY()
 
 public:
-	UMergeMeshesTool();
 
-	virtual void Setup() override;
-	virtual void OnShutdown(EToolShutdownType ShutdownType) override;
+	UMergeMeshesTool() {}
 
-	virtual void OnTick(float DeltaTime) override;
+protected:
 
-	virtual bool HasCancel() const override { return true; }
-	virtual bool HasAccept() const override { return true; }
-	virtual bool CanAccept() const override;
+	virtual void SetupProperties() override;
+	virtual void SaveProperties() override;
 
-	virtual void OnPropertyModified(UObject* PropertySet, FProperty* Property) override;
+	virtual FString GetCreatedAssetName() const override;
+	virtual FText GetActionName() const override;
+
+	virtual void ConvertInputsAndSetPreviewMaterials(bool bSetPreviewMesh) override;
 
 	// IDynamicMeshOperatorFactory API
 	virtual TUniquePtr<UE::Geometry::FDynamicMeshOperator> MakeNewOperator() override;
 
-protected:
 	UPROPERTY()
 	TObjectPtr<UMergeMeshesToolProperties> MergeProps;
 
-	UPROPERTY()
-	TObjectPtr<UMeshStatisticsProperties> MeshStatisticsProperties;
+};
 
-	UPROPERTY()
-	TObjectPtr<UOnAcceptHandleSourcesProperties> HandleSourcesProperties;
 
 
-	UPROPERTY()
-	TObjectPtr<UMeshOpPreviewWithBackgroundCompute> Preview;
 
-protected:
-	TArray<UE::Geometry::FVoxelMergeMeshesOp::FInputMesh> InputMeshes;
-	/** stash copies of the transforms and pointers to the meshes for consumption by merge Op*/
-	void CacheInputMeshes();
+UCLASS()
+class MESHMODELINGTOOLSEDITORONLYEXP_API UMergeMeshesToolBuilder : public UBaseCreateFromSelectedToolBuilder
+{
+	GENERATED_BODY()
 
-	/** quickly generate a low-quality result for display while the actual result is being computed. */
-	void CreateLowQualityPreview();
+public:
 
-	void GenerateAsset(const FDynamicMeshOpResult& Result);
+	virtual UMultiSelectionMeshEditingTool* CreateNewTool(const FToolBuilderState& SceneState) const override
+	{
+		return NewObject<UMergeMeshesTool>(SceneState.ToolManager);
+	}
+
+	virtual int32 MinComponentsSupported() const override { return 1; }
 };
+
+
diff --git a/Engine/Plugins/Experimental/MeshModelingToolsetExp/Source/MeshModelingToolsEditorOnlyExp/Public/VoxelCSGMeshesTool.h b/Engine/Plugins/Experimental/MeshModelingToolsetExp/Source/MeshModelingToolsEditorOnlyExp/Public/VoxelCSGMeshesTool.h
--- a/Engine/Plugins/Experimental/MeshModelingToolsetExp/Source/MeshModelingToolsEditorOnlyExp/Public/VoxelCSGMeshesTool.h
+++ b/Engine/Plugins/Experimental/MeshModelingToolsetExp/Source/MeshModelingToolsEditorOnlyExp/Public/VoxelCSGMeshesTool.h
@@ -3,27 +3,12 @@
 #pragma once
 
 #include "CoreMinimal.h"
-#include "InteractiveToolBuilder.h"
-#include "MeshOpPreviewHelpers.h"
-#include "BaseTools/MultiSelectionMeshEditingTool.h"
-#include "Properties/MeshStatisticsProperties.h"
-#include "PropertySets/OnAcceptProperties.h"
-#include "CompositionOps/VoxelBooleanMeshesOp.h"
+#include "UObject/NoExportTypes.h"
+#include "BaseTools/BaseVoxelTool.h"
 
-#include "VoxelCSGMeshesTool.generated.h"
 
-/**
- *
- */
-UCLASS()
-class MESHMODELINGTOOLSEDITORONLYEXP_API UVoxelCSGMeshesToolBuilder : public UMultiSelectionMeshEditingToolBuilder
-{
-	GENERATED_BODY()
+#include "VoxelCSGMeshesTool.generated.h"
 
-public:
-	virtual bool CanBuildTool(const FToolBuilderState& SceneState) const override;
-	virtual UMultiSelectionMeshEditingTool* CreateNewTool(const FToolBuilderState& SceneState) const override;
-};
 
 
 /**  */
@@ -54,76 +39,69 @@ class MESHMODELINGTOOLSEDITORONLYEXP_API UVoxelCSGMeshesToolProperties : public
 	GENERATED_BODY()
 public:
 	/** The type of operation  */
-	UPROPERTY(EditAnywhere, Category = Options)
+	UPROPERTY(EditAnywhere, Category = BooleanOptions)
 	EVoxelCSGOperation Operation = EVoxelCSGOperation::DifferenceAB;
 
 	/** The size of the geometry bounding box major axis measured in voxels.*/
-	UPROPERTY(EditAnywhere, Category = Options, meta = (UIMin = "8", UIMax = "1024", ClampMin = "8", ClampMax = "1024"))
+	UPROPERTY(EditAnywhere, Category = VoxelSettings, meta = (UIMin = "8", UIMax = "1024", ClampMin = "8", ClampMax = "1024"))
 	int32 VoxelCount = 128;
 
 	/** Remeshing adaptivity, prior to optional simplification */
-	UPROPERTY(EditAnywhere, Category = Options, meta = (UIMin = "0", UIMax = "1", ClampMin = "0", ClampMax = "1"))
+	UPROPERTY(EditAnywhere, Category = VoxelSettings, meta = (UIMin = "0", UIMax = "1", ClampMin = "0", ClampMax = "1"))
 	float MeshAdaptivity = 0.01f;
 
-	/** Offset when remeshing, note large offsets with high voxels counts will be slow */
-	UPROPERTY(EditAnywhere, Category = Options, meta = (UIMin = "-10", UIMax = "10", ClampMin = "-10", ClampMax = "10"))
+	/** Offset when remeshing, note large offsets with high voxels counts will be slow.  Hidden because this duplicates functionality of the voxel offset tool */
+	UPROPERTY()
 	float OffsetDistance = 0.0f;
 
 	/** Automatically simplify the result of voxel-based merge.*/
-	UPROPERTY(EditAnywhere, Category = Options)
+	UPROPERTY(EditAnywhere, Category = VoxelSettings)
 	bool bAutoSimplify = false;
 };
 
-
-
-
-/**
- *
- */
-UCLASS()
-class MESHMODELINGTOOLSEDITORONLYEXP_API UVoxelCSGMeshesTool : public UMultiSelectionMeshEditingTool, public UE::Geometry::IDynamicMeshOperatorFactory
+UCLASS()  
+class MESHMODELINGTOOLSEDITORONLYEXP_API UVoxelCSGMeshesTool : public UBaseVoxelTool
 {
 	GENERATED_BODY()
 
 public:
-	UVoxelCSGMeshesTool();
 
-	virtual void Setup() override;
-	virtual void OnShutdown(EToolShutdownType ShutdownType) override;
+	UVoxelCSGMeshesTool() {}
 
-	virtual void OnTick(float DeltaTime) override;
-	virtual void Render(IToolsContextRenderAPI* RenderAPI) override;
+protected:
+
+	virtual void SetupProperties() override;
+	virtual void SaveProperties() override;
 
-	virtual bool HasCancel() const override { return true; }
-	virtual bool HasAccept() const override { return true; }
-	virtual bool CanAccept() const override;
+	virtual FString GetCreatedAssetName() const override;
+	virtual FText GetActionName() const override;
 
-	virtual void OnPropertyModified(UObject* PropertySet, FProperty* Property) override;
+	virtual void ConvertInputsAndSetPreviewMaterials(bool bSetPreviewMesh) override;
 
 	// IDynamicMeshOperatorFactory API
 	virtual TUniquePtr<UE::Geometry::FDynamicMeshOperator> MakeNewOperator() override;
 
-
-protected:
 	UPROPERTY()
 	TObjectPtr<UVoxelCSGMeshesToolProperties> CSGProps;
 
-	UPROPERTY()
-	TObjectPtr<UMeshStatisticsProperties> MeshStatisticsProperties;
+};
 
-	UPROPERTY()
-	TObjectPtr<UOnAcceptHandleSourcesProperties> HandleSourcesProperties;
 
-	UPROPERTY()
-	TObjectPtr<UMeshOpPreviewWithBackgroundCompute> Preview;
 
-protected:
-	TArray<UE::Geometry::FVoxelBooleanMeshesOp::FInputMesh> InputMeshes;
-	/** stash copies of the transforms and pointers to the meshes for consumption by the CSG Op*/
-	void CacheInputMeshes();
 
-	/** quickly generate a low-quality result for display while the actual result is being computed. */
-	void CreateLowQualityPreview();
+UCLASS()
+class MESHMODELINGTOOLSEDITORONLYEXP_API UVoxelCSGMeshesToolBuilder : public UBaseCreateFromSelectedToolBuilder
+{
+	GENERATED_BODY()
+
+public:
+
+	virtual UMultiSelectionMeshEditingTool* CreateNewTool(const FToolBuilderState& SceneState) const override
+	{
+		return NewObject<UVoxelCSGMeshesTool>(SceneState.ToolManager);
+	}
 
-	void GenerateAsset(const FDynamicMeshOpResult& Result);
+	virtual TOptional<int32> MaxComponentsSupported() const override { return 2; }
+	virtual int32 MinComponentsSupported() const override { return 2; }
 };
+
diff --git a/Engine/Plugins/Experimental/ProxyLODPlugin/Source/ProxyLOD/Private/ProxyLODVolume.cpp b/Engine/Plugins/Experimental/ProxyLODPlugin/Source/ProxyLOD/Private/ProxyLODVolume.cpp
--- a/Engine/Plugins/Experimental/ProxyLODPlugin/Source/ProxyLOD/Private/ProxyLODVolume.cpp
+++ b/Engine/Plugins/Experimental/ProxyLODPlugin/Source/ProxyLOD/Private/ProxyLODVolume.cpp
@@ -21,7 +21,7 @@ THIRD_PARTY_INCLUDES_END
 
 #include "MeshDescription.h"
 
-typedef openvdb::math::Transform	OpenVDBTransform;
+typedef openvdb::math::Transform OpenVDBTransform;
 
 class FProxyLODVolumeImpl : public IProxyLODVolume
 {
@@ -286,7 +286,16 @@ public:
 
 	// Note this will become very slow if the isosurface is more than 3 voxel widths from 0.
 	// the tool that calls this should clamp the isosurface.
-	FVector ComputeUnion(const TArray<IVoxelBasedCSG::FPlacedMesh>& PlacedMeshArray, FMeshDescription& ResultMesh, double Adaptivity, double IsoSurface ) const override
+	virtual FVector ComputeUnion(const TArray<IVoxelBasedCSG::FPlacedMesh>& PlacedMeshArray, FMeshDescription& ResultMesh, double Adaptivity, double IsoSurface) const override
+	{
+		FInterrupter NullInterrupter;
+		FVector AverageTranslation;
+		ComputeUnion(NullInterrupter, PlacedMeshArray, ResultMesh, AverageTranslation, Adaptivity, IsoSurface);
+		return AverageTranslation;
+	}
+
+
+	virtual bool  ComputeUnion(IVoxelBasedCSG::FInterrupter& Interrupter,  const TArray<IVoxelBasedCSG::FPlacedMesh>& PlacedMeshArray, FMeshDescription& ResultMesh, FVector& AverageTranslation, double Adaptivity, double IsoSurface) const override
 	{
 
 		
@@ -299,16 +308,16 @@ public:
 		const double ExteriorVoxelWidth = FMath::Max( 2., IsoSurfaceInVoxels + 1.);
 		const double InteriorVoxelWidth = -FMath::Min(-2., IsoSurfaceInVoxels - 1);
 
-	
+		AverageTranslation = FVector(0.f, 0.f, 0.f);
+
 		const int32 NumMeshes = PlacedMeshArray.Num();
 		if (NumMeshes == 0)
 		{
-			return FVector(0.f, 0.f, 0.f);
+			return false;
 		}
 
 		// Find the average translation of all the meshes.
 		
-		FVector AverageTranslation(0.f, 0.f, 0.f);
 		for (int32 i = 0; i < NumMeshes; ++i)
 		{
 			AverageTranslation += PlacedMeshArray[i].Transform.GetTranslation();
@@ -361,72 +370,125 @@ public:
 		}
 
 		FPolygonSoup PolySoup(Adapters, VoxelSize);
-		openvdb::FloatGrid::Ptr SDFUnionVolume = openvdb::tools::meshToVolume<openvdb::FloatGrid>(PolySoup, PolySoup.Transform(), ExteriorVoxelWidth, InteriorVoxelWidth);
+		openvdb::FloatGrid::Ptr SDFUnionVolume = openvdb::tools::meshToVolume<openvdb::FloatGrid>(Interrupter, PolySoup, PolySoup.Transform(), ExteriorVoxelWidth, InteriorVoxelWidth);
 
+		const bool bSuccess = !Interrupter.wasInterrupted();
+		if (bSuccess)
+		{ 
+			// Convert the SDFUnionVolume to a mesh
+		
+			ConvertSDFToMesh(SDFUnionVolume, Adaptivity, IsoSurface, ResultMesh);
+		}
 
-		// Convert the SDFUnionVolume to a mesh
+		return bSuccess;
 		
-		ConvertSDFToMesh(SDFUnionVolume, Adaptivity, IsoSurface, ResultMesh);
+	}
 
+	
+	virtual FVector ComputeDifference(const FPlacedMesh& PlacedMeshA, const FPlacedMesh& PlacedMeshB, FMeshDescription& ResultMesh, double Adaptivity, double IsoSurface) const override
+	{
+		IVoxelBasedCSG::FInterrupter NullInterrupter;
+		FVector AverageTranslation;
+		ComputeDifference(NullInterrupter, PlacedMeshA, PlacedMeshB, ResultMesh, AverageTranslation, Adaptivity, IsoSurface);
 		return AverageTranslation;
 	}
 
-	FVector ComputeDifference(const FPlacedMesh& PlacedMeshA, const FPlacedMesh& PlacedMeshB, FMeshDescription& ResultMesh,  double Adaptivity, double IsoSurface) const override
+	virtual bool ComputeDifference(IVoxelBasedCSG::FInterrupter& Interrupter, const FPlacedMesh& PlacedMeshA, const FPlacedMesh& PlacedMeshB, FMeshDescription& ResultMesh, FVector& AverageTranslation, double Adaptivity, double IsoSurface) const override
 	{
 
 		// make SDFs
 		openvdb::FloatGrid::Ptr SDFVolumeA;
 		openvdb::FloatGrid::Ptr SDFVolumeB;
 
-		FVector AverageTranslation = GenerateVolumes(IsoSurface, PlacedMeshA, PlacedMeshB, SDFVolumeA, SDFVolumeB);
+		bool bSuccess = GenerateVolumes(Interrupter, IsoSurface, PlacedMeshA, PlacedMeshB, SDFVolumeA, SDFVolumeB, AverageTranslation);
 
-		// create the difference - result stored in volume A
-		openvdb::tools::csgDifference(*SDFVolumeA, *SDFVolumeB);
+		bSuccess = bSuccess && !Interrupter.wasInterrupted();
+		if (bSuccess)
+		{ 
+			// create the difference - result stored in volume A
+			openvdb::tools::csgDifference(*SDFVolumeA, *SDFVolumeB);
+		}
 
 		// convert the result
-		ConvertSDFToMesh(SDFVolumeA, Adaptivity, IsoSurface, ResultMesh);
+		bSuccess = bSuccess && !Interrupter.wasInterrupted();
+		if (bSuccess)
+		{ 
+			ConvertSDFToMesh(SDFVolumeA, Adaptivity, IsoSurface, ResultMesh);
+		}
+		return bSuccess;
+	}
 
+	
+	virtual FVector ComputeIntersection(const FPlacedMesh& PlacedMeshA, const FPlacedMesh& PlacedMeshB, FMeshDescription& ResultMesh, double Adaptivity, double IsoSurface) const override
+	{
+		IVoxelBasedCSG::FInterrupter NullInterrupter;
+		FVector AverageTranslation;
+		ComputeIntersection(NullInterrupter, PlacedMeshA, PlacedMeshB, ResultMesh, AverageTranslation, Adaptivity, IsoSurface);
 		return AverageTranslation;
 	}
 
-	FVector ComputeIntersection(const FPlacedMesh& PlacedMeshA, const FPlacedMesh& PlacedMeshB, FMeshDescription& ResultMesh, double Adaptivity, double IsoSurface) const override
+	virtual bool ComputeIntersection(IVoxelBasedCSG::FInterrupter& Interrupter,  const FPlacedMesh& PlacedMeshA, const FPlacedMesh& PlacedMeshB, FMeshDescription& ResultMesh, FVector& AverageTranslation, double Adaptivity, double IsoSurface) const override
 	{
 
 		// make SDFs
 		openvdb::FloatGrid::Ptr SDFVolumeA;
 		openvdb::FloatGrid::Ptr SDFVolumeB;
 
-		FVector AverageTranslation = GenerateVolumes(IsoSurface, PlacedMeshA, PlacedMeshB, SDFVolumeA, SDFVolumeB);
-
-
-		// create the difference - result stored in volume A
-		openvdb::tools::csgIntersection(*SDFVolumeA, *SDFVolumeB);
+		bool bSuccess = GenerateVolumes(Interrupter, IsoSurface, PlacedMeshA, PlacedMeshB, SDFVolumeA, SDFVolumeB, AverageTranslation);
+		bSuccess = bSuccess && !Interrupter.wasInterrupted();
+		if (bSuccess)
+		{ 
+			// create the difference - result stored in volume A
+			openvdb::tools::csgIntersection(*SDFVolumeA, *SDFVolumeB);
+		}
 
 		// convert the result
-		ConvertSDFToMesh(SDFVolumeA, Adaptivity, IsoSurface, ResultMesh);
+		bSuccess = bSuccess && !Interrupter.wasInterrupted();
+		if (bSuccess)
+		{
+			// convert the result
+			ConvertSDFToMesh(SDFVolumeA, Adaptivity, IsoSurface, ResultMesh);
+		}
+		
+		return bSuccess;
+	}
 
+	
+	
+	virtual FVector ComputeUnion(const FPlacedMesh& PlacedMeshA, const FPlacedMesh& PlacedMeshB, FMeshDescription& ResultMesh, double Adaptivity, double IsoSurface) const override
+	{
+		FInterrupter NullInterrupter;
+		FVector AverageTranslation;
+		ComputeUnion(NullInterrupter, PlacedMeshA, PlacedMeshB, ResultMesh, AverageTranslation, Adaptivity, IsoSurface);
 		return AverageTranslation;
 	}
-
-	FVector ComputeUnion(const FPlacedMesh& PlacedMeshA, const FPlacedMesh& PlacedMeshB, FMeshDescription& ResultMesh, double Adaptivity, double IsoSurface) const override
+	
+	virtual bool ComputeUnion(IVoxelBasedCSG::FInterrupter& Interrupter, const FPlacedMesh& PlacedMeshA, const FPlacedMesh& PlacedMeshB, FMeshDescription& ResultMesh, FVector& AverageTranslation, double Adaptivity, double IsoSurface) const override
 	{
 
 		// make SDFs
 		openvdb::FloatGrid::Ptr SDFVolumeA;
 		openvdb::FloatGrid::Ptr SDFVolumeB;
 
-		FVector AverageTranslation = GenerateVolumes(IsoSurface, PlacedMeshA, PlacedMeshB, SDFVolumeA, SDFVolumeB);
-
-
-		// create the difference - result stored in volume A
-		openvdb::tools::csgUnion(*SDFVolumeA, *SDFVolumeB);
+		bool bSuccess = GenerateVolumes(Interrupter, IsoSurface, PlacedMeshA, PlacedMeshB, SDFVolumeA, SDFVolumeB, AverageTranslation);
+		bSuccess = bSuccess && !Interrupter.wasInterrupted();
+		if (bSuccess)
+		{ 
+			// create the difference - result stored in volume A
+			openvdb::tools::csgUnion(*SDFVolumeA, *SDFVolumeB);
+		}
 
 		// convert the result
-		ConvertSDFToMesh(SDFVolumeA, Adaptivity, IsoSurface, ResultMesh);
-
-		return AverageTranslation;
+		bSuccess = bSuccess && !Interrupter.wasInterrupted();
+		if (bSuccess)
+		{	
+			ConvertSDFToMesh(SDFVolumeA, Adaptivity, IsoSurface, ResultMesh);
+		}
+		return bSuccess;
 	}
 
+	
+
 private:
 
 	void ConvertSDFToMesh(openvdb::FloatGrid::Ptr SDFVolume, double Adaptivity, double IsoSurface, FMeshDescription& ResultMesh) const
@@ -441,7 +503,7 @@ private:
 		ProxyLOD::ConvertMesh(AOSMeshedVolume, ResultMesh);
 	}
 
-	FVector GenerateVolumes(const double IsoSurface, const FPlacedMesh& PlacedMeshA, const FPlacedMesh& PlacedMeshB, openvdb::FloatGrid::Ptr& VolumeA, openvdb::FloatGrid::Ptr& VolumeB) const 
+	bool  GenerateVolumes(IVoxelBasedCSG::FInterrupter& Interrupter, const double IsoSurface, const FPlacedMesh& PlacedMeshA, const FPlacedMesh& PlacedMeshB, openvdb::FloatGrid::Ptr& VolumeA, openvdb::FloatGrid::Ptr& VolumeB, FVector& AverageTranslation) const	
 	{
 		const FMeshDescription& MeshA = *PlacedMeshA.Mesh;
 		const FMeshDescription& MeshB = *PlacedMeshB.Mesh;
@@ -455,7 +517,7 @@ private:
 		const double InteriorVoxelWidth = -FMath::Min(-2., IsoSurfaceInVoxels - 1);
 
 		// The average translation of the two meshes.
-		FVector AverageTranslation = PlacedMeshA.Transform.GetTranslation() + PlacedMeshB.Transform.GetTranslation();
+		AverageTranslation = PlacedMeshA.Transform.GetTranslation() + PlacedMeshB.Transform.GetTranslation();
 		AverageTranslation *= 0.5f;
 
 		// Fill the SDFUnionVolume
@@ -493,16 +555,26 @@ private:
 		OpenVDBTransform::Ptr TargetXForm = OpenVDBTransform::createLinearTransform(VoxelSize);
 
 		// make SDFs
-		openvdb::FloatGrid::Ptr SDFVolumeA = openvdb::tools::meshToVolume<openvdb::FloatGrid>(AdapterA, *TargetXForm, ExteriorVoxelWidth, InteriorVoxelWidth);
-		openvdb::FloatGrid::Ptr SDFVolumeB = openvdb::tools::meshToVolume<openvdb::FloatGrid>(AdapterB, *TargetXForm, ExteriorVoxelWidth, InteriorVoxelWidth);
+		openvdb::FloatGrid::Ptr SDFVolumeA = openvdb::tools::meshToVolume<openvdb::FloatGrid>(Interrupter, AdapterA, *TargetXForm, ExteriorVoxelWidth, InteriorVoxelWidth);
+		openvdb::FloatGrid::Ptr SDFVolumeB = openvdb::tools::meshToVolume<openvdb::FloatGrid>(Interrupter, AdapterB, *TargetXForm, ExteriorVoxelWidth, InteriorVoxelWidth);
 
 		VolumeA = SDFVolumeA;
 		VolumeB = SDFVolumeB;
 
-		return AverageTranslation;
+		const bool bSuccess = !Interrupter.wasInterrupted();
+		return bSuccess;
 
 	}
 
+
+	FVector GenerateVolumes(const double IsoSurface, const FPlacedMesh& PlacedMeshA, const FPlacedMesh& PlacedMeshB, openvdb::FloatGrid::Ptr& VolumeA, openvdb::FloatGrid::Ptr& VolumeB) const
+	{
+		FVector AverageTranslation;
+		IVoxelBasedCSG::FInterrupter NullInterrupter;
+		GenerateVolumes(NullInterrupter, IsoSurface, PlacedMeshA, PlacedMeshB, VolumeA, VolumeB, AverageTranslation);
+		return AverageTranslation;
+	}
+
 	OpenVDBTransform::Ptr XForm;
 };
 
diff --git a/Engine/Plugins/Experimental/ProxyLODPlugin/Source/ProxyLOD/Public/ProxyLODVolume.h b/Engine/Plugins/Experimental/ProxyLODPlugin/Source/ProxyLOD/Public/ProxyLODVolume.h
--- a/Engine/Plugins/Experimental/ProxyLODPlugin/Source/ProxyLOD/Public/ProxyLODVolume.h
+++ b/Engine/Plugins/Experimental/ProxyLODPlugin/Source/ProxyLOD/Public/ProxyLODVolume.h
@@ -3,6 +3,7 @@
 #pragma once
 
 #include "CoreMinimal.h"
+#include "Math/Vector.h"
 
 struct FMeshMergeData;
 struct FInstancedMeshMergeData;
@@ -96,14 +97,33 @@ public :
 		FTransform       Transform;
 	};
 
+	// the interrupter interface. 
+	struct FInterrupter
+	{
+	
+		FInterrupter() = default;
+		virtual ~FInterrupter() = default;
+		// templated openvdb code requires these functions.
+		virtual void start(const char* name = nullptr) { (void)name; }
+		virtual void end() {};
+		virtual bool wasInterrupted(int percent = -1) { (void)percent; return false; }
+	};
+
 	virtual double GetVoxelSize() const = 0;
 	virtual void SetVoxelSize(double VolexSize) = 0;
 
 	// Will destroy UVs and other attributes  Returns the average location of the input meshes 
 	virtual FVector ComputeUnion(const TArray<IVoxelBasedCSG::FPlacedMesh>& MeshArray, FMeshDescription& ResultMesh, double Adaptivity = 0.1, double IsoSurfcae = 0.) const = 0;
 	
-	// We could make this keep the UVs and other attributes from the AMesh..
 	virtual FVector ComputeDifference(const FPlacedMesh& PlacedMeshA, const FPlacedMesh& PlacedMeshB, FMeshDescription& ResultMesh, double Adaptivity, double IsoSurface) const = 0;
 	virtual FVector ComputeIntersection(const FPlacedMesh& PlacedMeshA, const FPlacedMesh& PlacedMeshB, FMeshDescription& ResultMesh, double Adaptivity, double IsoSurface) const = 0;
 	virtual FVector ComputeUnion(const FPlacedMesh& PlacedMeshA, const FPlacedMesh& PlacedMeshB, FMeshDescription& ResultMesh, double Adaptivity, double IsoSurface) const = 0;
+	
+	
+	// interruptible versions. return true on success.  
+	virtual bool ComputeUnion(IVoxelBasedCSG::FInterrupter& Interrupter, const TArray<IVoxelBasedCSG::FPlacedMesh>& MeshArray, FMeshDescription& ResultMesh, FVector& AverageTranslation, double Adaptivity, double IsoSurfcae) const = 0;
+	
+	virtual bool ComputeDifference(IVoxelBasedCSG::FInterrupter& Interrupter, const FPlacedMesh& PlacedMeshA, const FPlacedMesh& PlacedMeshB, FMeshDescription& ResultMesh, FVector& AverageTranslation, double Adaptivity, double IsoSurface) const = 0;
+	virtual bool ComputeIntersection(IVoxelBasedCSG::FInterrupter& Interrupter, const FPlacedMesh& PlacedMeshA, const FPlacedMesh& PlacedMeshB, FMeshDescription& ResultMesh, FVector& AverageTranslation, double Adaptivity, double IsoSurface) const = 0;
+	virtual bool ComputeUnion(IVoxelBasedCSG::FInterrupter& Interrupter, const FPlacedMesh& PlacedMeshA, const FPlacedMesh& PlacedMeshB, FMeshDescription& ResultMesh, FVector& AverageTranslation, double Adaptivity, double IsoSurface) const = 0;
 };
\ No newline at end of file
diff --git a/Engine/Plugins/Runtime/MeshModelingToolset/Source/ModelingOperatorsEditorOnly/Private/CompositionOps/VoxelBooleanMeshesOp.cpp b/Engine/Plugins/Runtime/MeshModelingToolset/Source/ModelingOperatorsEditorOnly/Private/CompositionOps/VoxelBooleanMeshesOp.cpp
--- a/Engine/Plugins/Runtime/MeshModelingToolset/Source/ModelingOperatorsEditorOnly/Private/CompositionOps/VoxelBooleanMeshesOp.cpp
+++ b/Engine/Plugins/Runtime/MeshModelingToolset/Source/ModelingOperatorsEditorOnly/Private/CompositionOps/VoxelBooleanMeshesOp.cpp
@@ -5,10 +5,12 @@
 #include "MeshDescription.h"
 #include "StaticMeshAttributes.h"
 #include "MeshDescriptionToDynamicMesh.h"
+#include "DynamicMeshToMeshDescription.h"
 #include "MeshSimplification.h"
 #include "MeshConstraintsUtil.h"
 #include "CleaningOps/EditNormalsOp.h"
 
+
 // The ProxyLOD plugin is currently only available on Windows. On other platforms we will make this a no-op.
 #if WITH_PROXYLOD
 #include "ProxyLODVolume.h"
@@ -19,6 +21,21 @@ using namespace UE::Geometry;
 void FVoxelBooleanMeshesOp::CalculateResult(FProgressCancel* Progress)
 {
 #if WITH_PROXYLOD
+	
+	struct FVoxelBoolInterrupter : IVoxelBasedCSG::FInterrupter
+	{
+		FVoxelBoolInterrupter(FProgressCancel* ProgressCancel) : Progress(ProgressCancel) {}
+		FProgressCancel* Progress;
+		virtual ~FVoxelBoolInterrupter(){}
+		virtual bool wasInterrupted(int percent = -1) override final
+		{
+			bool Cancelled = Progress && Progress->Cancelled();
+			return Cancelled;
+		}
+
+	} Interrupter(Progress);
+
+
 	FMeshDescription ResultMeshDescription;
 	FStaticMeshAttributes Attributes(ResultMeshDescription);
 	Attributes.Register();
@@ -49,35 +66,41 @@ void FVoxelBooleanMeshesOp::CalculateResult(FProgressCancel* Progress)
 	TUniquePtr<IVoxelBasedCSG> VoxelCSGTool = IVoxelBasedCSG::CreateCSGTool(VoxelSizeD);
 	FVector MergedOrigin;
 
-	TArray<IVoxelBasedCSG::FPlacedMesh> PlacedMeshes;
-	for (const FInputMesh& InputMesh : InputMeshArray)
+	// VoxelCSGTool needs MeshDescription
+	FMeshDescription MeshDescriptions[2];
+	for (int i = 0; i < 2; ++i)
 	{
-		IVoxelBasedCSG::FPlacedMesh PlacedMesh;
-		PlacedMesh.Mesh = InputMesh.Mesh;
-		PlacedMesh.Transform = InputMesh.Transform;
-		PlacedMeshes.Add(PlacedMesh);
+
+		FStaticMeshAttributes AddAttributes(MeshDescriptions[i]);
+		AddAttributes.Register();
+		
+		FConversionToMeshDescriptionOptions ToMeshDescriptionOptions;
+		ToMeshDescriptionOptions.bSetPolyGroups = false;
+		FDynamicMeshToMeshDescription DynamicMeshToMeshDescription(ToMeshDescriptionOptions);
+
+		DynamicMeshToMeshDescription.Convert(Meshes[i].Get(), MeshDescriptions[i]);
 	}
 
-	IVoxelBasedCSG::FPlacedMesh& A = PlacedMeshes[0];
-	IVoxelBasedCSG::FPlacedMesh& B = PlacedMeshes[1];
+	IVoxelBasedCSG::FPlacedMesh A(&MeshDescriptions[0], Transforms[0]);
+	IVoxelBasedCSG::FPlacedMesh B(&MeshDescriptions[1], Transforms[1]);
 
 	switch (Operation)
 	{
 	case EBooleanOperation::DifferenceAB:
 		CSGIsoSurface = 0.;
-		MergedOrigin = VoxelCSGTool->ComputeDifference(A, B, ResultMeshDescription, AdaptivityD, CSGIsoSurface);
+		VoxelCSGTool->ComputeDifference(Interrupter, A, B, ResultMeshDescription, MergedOrigin, AdaptivityD, CSGIsoSurface);
 		break;
 	case EBooleanOperation::DifferenceBA:
 		CSGIsoSurface = 0.;
-		MergedOrigin = VoxelCSGTool->ComputeDifference(B, A, ResultMeshDescription, AdaptivityD, CSGIsoSurface);
+		VoxelCSGTool->ComputeDifference(Interrupter, B, A, ResultMeshDescription, MergedOrigin, AdaptivityD, CSGIsoSurface);
 		break;
 	case EBooleanOperation::Intersect:
 		CSGIsoSurface = FMath::Clamp(IsoSurfaceD, -MaxIsoOffset, MaxIsoOffset);
-		MergedOrigin = VoxelCSGTool->ComputeIntersection(A, B, ResultMeshDescription, AdaptivityD, CSGIsoSurface);
+		VoxelCSGTool->ComputeIntersection(Interrupter, A, B, ResultMeshDescription, MergedOrigin, AdaptivityD, CSGIsoSurface);
 		break;
 	case EBooleanOperation::Union:
 		CSGIsoSurface = FMath::Clamp(IsoSurfaceD, 0., MaxIsoOffset);
-		MergedOrigin = VoxelCSGTool->ComputeUnion(A, B, ResultMeshDescription, AdaptivityD, CSGIsoSurface);
+		VoxelCSGTool->ComputeUnion(Interrupter, A, B, ResultMeshDescription, MergedOrigin, AdaptivityD, CSGIsoSurface);
 		break;
 	default:
 		check(0);
@@ -85,7 +108,7 @@ void FVoxelBooleanMeshesOp::CalculateResult(FProgressCancel* Progress)
 
 	ResultTransform = FTransform3d(MergedOrigin);
 
-	if (Progress && Progress->Cancelled())
+	if (Interrupter.wasInterrupted())
 	{
 		return;
 	}
@@ -102,7 +125,7 @@ void FVoxelBooleanMeshesOp::CalculateResult(FProgressCancel* Progress)
 
 		for (int32 i = 0; i < NumRemeshes; ++i)
 		{
-			if (Progress && Progress->Cancelled())
+			if (Interrupter.wasInterrupted())
 			{
 				return;
 			}
@@ -115,17 +138,19 @@ void FVoxelBooleanMeshesOp::CalculateResult(FProgressCancel* Progress)
 			SourceMeshes.Add(PlacedMesh);
 
 			// union of only one mesh. we are using this to voxelize and remesh with and offset.
-			MergedOrigin = VoxelCSGTool->ComputeUnion(SourceMeshes, ResultMeshDescription, AdaptivityD, DeltaIsoSurface);
+			VoxelCSGTool->ComputeUnion(Interrupter, SourceMeshes, ResultMeshDescription, MergedOrigin, AdaptivityD, DeltaIsoSurface);
 
 			ResultTransform = FTransform3d(MergedOrigin);
 		}
 	}
-
-	// Convert to dynamic mesh
-	FMeshDescriptionToDynamicMesh Converter;
-	Converter.Convert(&ResultMeshDescription, *ResultMesh);
-
-	if (Progress && Progress->Cancelled())
+	
+	if(!Interrupter.wasInterrupted())
+	{
+		// Convert to dynamic mesh
+		FMeshDescriptionToDynamicMesh Converter;
+		Converter.Convert(&ResultMeshDescription, *ResultMesh);
+	}
+	else
 	{
 		return;
 	}
@@ -148,7 +173,7 @@ void FVoxelBooleanMeshesOp::CalculateResult(FProgressCancel* Progress)
 	}
 
 
-	if (bFixNormals)
+	if (bFixNormals && !Interrupter.wasInterrupted())
 	{
 		TSharedPtr<FDynamicMesh3, ESPMode::ThreadSafe> OpResultMesh(ExtractResult().Release()); // moved the unique pointer
 
@@ -169,13 +194,11 @@ void FVoxelBooleanMeshesOp::CalculateResult(FProgressCancel* Progress)
 
 		ResultMesh = EditNormalsOp.ExtractResult(); // return the edit normals operator copy to this tool.
 	}
-
+	
 #else	// WITH_PROXYLOD
-
-	FMeshDescriptionToDynamicMesh Converter;
-	Converter.Convert(InputMeshArray[0].Mesh, *ResultMesh);
-	ResultTransform = FTransform3d(InputMeshArray[0].Transform);
-
+		// just copy the first input
+		*ResultMesh = *Meshes[0];
+		ResultTransform = Transforms[0];
 #endif	// WITH_PROXYLOD
 
 }
@@ -183,45 +206,23 @@ void FVoxelBooleanMeshesOp::CalculateResult(FProgressCancel* Progress)
 float FVoxelBooleanMeshesOp::ComputeVoxelSize() const 
 { 
 	float Size = 0.f;
-	for (int32 i = 0; i < InputMeshArray.Num(); ++i)
+	auto GrowSize = [&Size, this](const FDynamicMesh3& DynamicMesh, const FTransformSRT3d& Xform)
 	{
-		//Bounding box
-		FVector BBoxMin(FLT_MAX, FLT_MAX, FLT_MAX);
-		FVector BBoxMax(-FLT_MAX, -FLT_MAX, -FLT_MAX);
-
-		// Use the scale define by the transform.  We don't care about actual placement
-		// in the world for this.
-		FVector Scale = InputMeshArray[i].Transform.GetScale3D();
-		const FMeshDescription&  MeshDescription = *InputMeshArray[i].Mesh;
-
-		TArrayView<const FVector3f> VertexPositions = MeshDescription.GetVertexPositions().GetRawArray();
-		for (const FVertexID VertexID : MeshDescription.Vertices().GetElementIDs())
-		{
-			const FVector3f& Pos = VertexPositions[VertexID];
-            
-			BBoxMin.X = FMath::Min(BBoxMin.X, Pos.X);
-			BBoxMin.Y = FMath::Min(BBoxMin.Y, Pos.Y);
-			BBoxMin.Z = FMath::Min(BBoxMin.Z, Pos.Z);
-
-			BBoxMax.X = FMath::Max(BBoxMax.X, Pos.X);
-			BBoxMax.Y = FMath::Max(BBoxMax.Y, Pos.Y);
-			BBoxMax.Z = FMath::Max(BBoxMax.Z, Pos.Z);
-		}
-
-		// The size of the BBox in each direction
-		FVector Extents(BBoxMax.X - BBoxMin.X, BBoxMax.Y - BBoxMin.Y, BBoxMax.Z - BBoxMin.Z);
-		
+		FAxisAlignedBox3d AABB = DynamicMesh.GetBounds(true);
+		FVector Scale = Xform.GetScale3D();
+		FVector Extents = 2. * AABB.Extents();
 		// Scale with the local space scale.
 		Extents.X = Extents.X * FMath::Abs(Scale.X);
 		Extents.Y = Extents.Y * FMath::Abs(Scale.Y);
 		Extents.Z = Extents.Z * FMath::Abs(Scale.Z);
 
 		float MajorAxisSize = FMath::Max3(Extents.X, Extents.Y, Extents.Z);
-
-
 		Size = FMath::Max(MajorAxisSize / VoxelCount, Size);
+	};
 
-	}
+	GrowSize(*Meshes[0], Transforms[0]);
+	GrowSize(*Meshes[1], Transforms[1]);
+	
 	
 	return Size;
 
diff --git a/Engine/Plugins/Runtime/MeshModelingToolset/Source/ModelingOperatorsEditorOnly/Private/CompositionOps/VoxelMergeMeshesOp.cpp b/Engine/Plugins/Runtime/MeshModelingToolset/Source/ModelingOperatorsEditorOnly/Private/CompositionOps/VoxelMergeMeshesOp.cpp
--- a/Engine/Plugins/Runtime/MeshModelingToolset/Source/ModelingOperatorsEditorOnly/Private/CompositionOps/VoxelMergeMeshesOp.cpp
+++ b/Engine/Plugins/Runtime/MeshModelingToolset/Source/ModelingOperatorsEditorOnly/Private/CompositionOps/VoxelMergeMeshesOp.cpp
@@ -5,6 +5,7 @@
 #include "MeshDescription.h"
 #include "StaticMeshAttributes.h"
 #include "MeshDescriptionToDynamicMesh.h"
+#include "DynamicMeshToMeshDescription.h"
 #include "MeshSimplification.h"
 #include "MeshConstraintsUtil.h"
 #include "CleaningOps/EditNormalsOp.h"
@@ -20,6 +21,19 @@ void FVoxelMergeMeshesOp::CalculateResult(FProgressCancel* Progress)
 {
 #if WITH_PROXYLOD
 
+	struct FVoxelBoolInterrupter : IVoxelBasedCSG::FInterrupter
+	{
+		FVoxelBoolInterrupter(FProgressCancel* ProgressCancel) : Progress(ProgressCancel) {}
+		FProgressCancel* Progress;
+		virtual ~FVoxelBoolInterrupter() {}
+		virtual bool wasInterrupted(int percent = -1) override final
+		{
+			bool Cancelled = Progress && Progress->Cancelled();
+			return Cancelled;
+		}
+
+	} Interrupter(Progress);
+
 	FMeshDescription MergedMeshesDescription;
 	FStaticMeshAttributes Attributes(MergedMeshesDescription);
 	Attributes.Register();
@@ -33,22 +47,37 @@ void FVoxelMergeMeshesOp::CalculateResult(FProgressCancel* Progress)
 	// Create CSGTool and merge the meshes.
 	TUniquePtr<IVoxelBasedCSG> VoxelCSGTool = IVoxelBasedCSG::CreateCSGTool(VoxelSizeD);
 
+	// voxel tool requires mesh description
+	TArray<FMeshDescription> MeshDescriptions;
+	MeshDescriptions.SetNum(Meshes.Num());
+	for (int i = 0; i < Meshes.Num(); ++i)
+	{
+		FStaticMeshAttributes AddAttributes(MeshDescriptions[i]);
+		AddAttributes.Register();
+
+		FConversionToMeshDescriptionOptions ToMeshDescriptionOptions;
+		ToMeshDescriptionOptions.bSetPolyGroups = false;
+		FDynamicMeshToMeshDescription DynamicMeshToMeshDescription(ToMeshDescriptionOptions);
+
+		DynamicMeshToMeshDescription.Convert(Meshes[i].Get(), MeshDescriptions[i]);
+	}
+
 	TArray<IVoxelBasedCSG::FPlacedMesh> PlacedMeshes;
-	for (const FInputMesh& InputMesh : InputMeshArray)
+	PlacedMeshes.Reserve(Meshes.Num());
+	for (int i = 0; i < MeshDescriptions.Num(); ++i)
 	{
-		IVoxelBasedCSG::FPlacedMesh PlacedMesh;
-		PlacedMesh.Mesh = InputMesh.Mesh;
-		PlacedMesh.Transform = InputMesh.Transform;
+		IVoxelBasedCSG::FPlacedMesh PlacedMesh(&MeshDescriptions[i], Transforms[i]);
 		PlacedMeshes.Add(PlacedMesh);
 	}
 
 	// world space units.
 	const double MaxIsoOffset = 2 * VoxelSizeD;
 	const double CSGIsoSurface = FMath::Clamp(IsoSurfaceD, 0., MaxIsoOffset); // the interior distance values maybe messed up when doing a union.
-	FVector MergedOrigin = VoxelCSGTool->ComputeUnion(PlacedMeshes, MergedMeshesDescription, AdaptivityD, CSGIsoSurface);
+	FVector MergedOrigin; 
+	bool bSuccess = VoxelCSGTool->ComputeUnion(Interrupter, PlacedMeshes, MergedMeshesDescription, MergedOrigin, AdaptivityD, CSGIsoSurface);
 	ResultTransform = FTransform3d(MergedOrigin);
 
-	if (Progress && Progress->Cancelled())
+	if (!bSuccess || Interrupter.wasInterrupted())
 	{
 		return;
 	}
@@ -65,7 +94,7 @@ void FVoxelMergeMeshesOp::CalculateResult(FProgressCancel* Progress)
 
 		for (int32 i = 0; i < NumRemeshes; ++i)
 		{
-			if (Progress && Progress->Cancelled())
+			if (Interrupter.wasInterrupted())
 			{
 				return;
 			}
@@ -78,17 +107,19 @@ void FVoxelMergeMeshesOp::CalculateResult(FProgressCancel* Progress)
 			SourceMeshes.Add(PlacedMesh);
 
 			// union of only one mesh. we are using this to voxelize and remesh with and offset.
-			MergedOrigin = VoxelCSGTool->ComputeUnion(SourceMeshes, MergedMeshesDescription, AdaptivityD, DeltaIsoSurface);
+			VoxelCSGTool->ComputeUnion(Interrupter, SourceMeshes, MergedMeshesDescription, MergedOrigin, AdaptivityD, DeltaIsoSurface);
 
 			ResultTransform = FTransform3d(MergedOrigin);
 		}
 	}
 
-	// Convert to dynamic mesh
-	FMeshDescriptionToDynamicMesh Converter;
-	Converter.Convert(&MergedMeshesDescription, *ResultMesh);
-
-	if (Progress && Progress->Cancelled())
+	if (!Interrupter.wasInterrupted())
+	{ 
+		// Convert to dynamic mesh
+		FMeshDescriptionToDynamicMesh Converter;
+		Converter.Convert(&MergedMeshesDescription, *ResultMesh);
+	}
+	else
 	{
 		return;
 	}
@@ -111,7 +142,7 @@ void FVoxelMergeMeshesOp::CalculateResult(FProgressCancel* Progress)
 
 	}
 
-	if (bFixNormals)
+	if (bFixNormals  && !Interrupter.wasInterrupted())
 	{
 
 		TSharedPtr<FDynamicMesh3, ESPMode::ThreadSafe> OpResultMesh(ExtractResult().Release()); // moved the unique pointer
@@ -136,11 +167,9 @@ void FVoxelMergeMeshesOp::CalculateResult(FProgressCancel* Progress)
 
 
 #else	// WITH_PROXYLOD
-
-	FMeshDescriptionToDynamicMesh Converter;
-	Converter.Convert(InputMeshArray[0].Mesh, *ResultMesh);
-	ResultTransform = FTransform3d(InputMeshArray[0].Transform);
-
+	// just copy the first input
+	*ResultMesh = *Meshes[0];
+	ResultTransform = Transforms[0];
 #endif	// WITH_PROXYLOD
 
 }
@@ -149,44 +178,22 @@ void FVoxelMergeMeshesOp::CalculateResult(FProgressCancel* Progress)
 float FVoxelMergeMeshesOp::ComputeVoxelSize() const
 {
 	float Size = 0.f;
-	for (int32 i = 0; i < InputMeshArray.Num(); ++i)
+	auto GrowSize = [&Size, this](const FDynamicMesh3& DynamicMesh, const FTransformSRT3d& Xform)
 	{
-		//Bounding box
-		FVector BBoxMin(FLT_MAX, FLT_MAX, FLT_MAX);
-		FVector BBoxMax(-FLT_MAX, -FLT_MAX, -FLT_MAX);
-
-		// Use the scale define by the transform.  We don't care about actual placement
-		// in the world for this.
-		FVector Scale = InputMeshArray[i].Transform.GetScale3D();
-		const FMeshDescription&  MeshDescription = *InputMeshArray[i].Mesh;
-
-		TArrayView<const FVector3f> VertexPositions = MeshDescription.GetVertexPositions().GetRawArray();
-		for (const FVertexID VertexID : MeshDescription.Vertices().GetElementIDs())
-		{
-			const FVector3f& Pos = VertexPositions[VertexID];
-
-			BBoxMin.X = FMath::Min(BBoxMin.X, Pos.X);
-			BBoxMin.Y = FMath::Min(BBoxMin.Y, Pos.Y);
-			BBoxMin.Z = FMath::Min(BBoxMin.Z, Pos.Z);
-
-			BBoxMax.X = FMath::Max(BBoxMax.X, Pos.X);
-			BBoxMax.Y = FMath::Max(BBoxMax.Y, Pos.Y);
-			BBoxMax.Z = FMath::Max(BBoxMax.Z, Pos.Z);
-		}
-
-		// The size of the BBox in each direction
-		FVector Extents(BBoxMax.X - BBoxMin.X, BBoxMax.Y - BBoxMin.Y, BBoxMax.Z - BBoxMin.Z);
-
+		FAxisAlignedBox3d AABB = DynamicMesh.GetBounds(true);
+		FVector Scale = Xform.GetScale3D();
+		FVector Extents = 2. * AABB.Extents();
 		// Scale with the local space scale.
 		Extents.X = Extents.X * FMath::Abs(Scale.X);
 		Extents.Y = Extents.Y * FMath::Abs(Scale.Y);
 		Extents.Z = Extents.Z * FMath::Abs(Scale.Z);
 
 		float MajorAxisSize = FMath::Max3(Extents.X, Extents.Y, Extents.Z);
-
-
 		Size = FMath::Max(MajorAxisSize / VoxelCount, Size);
-
+	};
+	for (int32 i = 0; i < Meshes.Num(); ++i)
+	{
+		GrowSize(*Meshes[i], Transforms[i]);
 	}
 
 	return Size;
diff --git a/Engine/Plugins/Runtime/MeshModelingToolset/Source/ModelingOperatorsEditorOnly/Public/CompositionOps/VoxelBooleanMeshesOp.h b/Engine/Plugins/Runtime/MeshModelingToolset/Source/ModelingOperatorsEditorOnly/Public/CompositionOps/VoxelBooleanMeshesOp.h
--- a/Engine/Plugins/Runtime/MeshModelingToolset/Source/ModelingOperatorsEditorOnly/Public/CompositionOps/VoxelBooleanMeshesOp.h
+++ b/Engine/Plugins/Runtime/MeshModelingToolset/Source/ModelingOperatorsEditorOnly/Public/CompositionOps/VoxelBooleanMeshesOp.h
@@ -5,6 +5,8 @@
 #include "CoreMinimal.h"
 #include "Util/ProgressCancel.h"
 #include "ModelingOperators.h"
+#include "DynamicMesh/DynamicMesh3.h"
+#include "TransformTypes.h"
 
 struct FMeshDescription;
 
@@ -27,14 +29,11 @@ public:
 		Union = 3
 	};
 
-	struct FInputMesh
-	{
-		const FMeshDescription* Mesh;
-		FTransform Transform;
-	};
-
+	
 	// inputs
-	TArray<FInputMesh> InputMeshArray;
+	TArray<TSharedPtr<const FDynamicMesh3, ESPMode::ThreadSafe>> Meshes;
+	TArray<FTransformSRT3d> Transforms; // 1:1 with Meshes
+
 	int32 VoxelCount = 128;
 	double VoxelSizeD = 1.0;
 	double AdaptivityD = 0;
diff --git a/Engine/Plugins/Runtime/MeshModelingToolset/Source/ModelingOperatorsEditorOnly/Public/CompositionOps/VoxelMergeMeshesOp.h b/Engine/Plugins/Runtime/MeshModelingToolset/Source/ModelingOperatorsEditorOnly/Public/CompositionOps/VoxelMergeMeshesOp.h
--- a/Engine/Plugins/Runtime/MeshModelingToolset/Source/ModelingOperatorsEditorOnly/Public/CompositionOps/VoxelMergeMeshesOp.h
+++ b/Engine/Plugins/Runtime/MeshModelingToolset/Source/ModelingOperatorsEditorOnly/Public/CompositionOps/VoxelMergeMeshesOp.h
@@ -5,6 +5,8 @@
 #include "CoreMinimal.h"
 #include "Util/ProgressCancel.h"
 #include "ModelingOperators.h"
+#include "DynamicMesh/DynamicMesh3.h"
+#include "TransformTypes.h"
 
 struct FMeshDescription;
 
@@ -19,14 +21,10 @@ class MODELINGOPERATORSEDITORONLY_API FVoxelMergeMeshesOp : public FDynamicMeshO
 public:
 	virtual ~FVoxelMergeMeshesOp() {}
 
-	struct FInputMesh
-	{
-		const FMeshDescription* Mesh;
-		FTransform Transform;
-	};
-
 	// inputs
-	TArray<FInputMesh> InputMeshArray;
+	TArray<TSharedPtr<const FDynamicMesh3, ESPMode::ThreadSafe>> Meshes;
+	TArray<FTransformSRT3d> Transforms; // 1:1 with Meshes
+
 	int32 VoxelCount = 128;
 	double VoxelSizeD = 1.0;
 	double AdaptivityD = 0;
