Fixed mismatched copy constructor/assignment operator in FPositionNormalVertex.

diff --git a/Engine/Plugins/Experimental/ProxyLODPlugin/Source/ProxyLOD/Private/ProxyLODVertexTypes.h b/Engine/Plugins/Experimental/ProxyLODPlugin/Source/ProxyLOD/Private/ProxyLODVertexTypes.h
--- a/Engine/Plugins/Experimental/ProxyLODPlugin/Source/ProxyLOD/Private/ProxyLODVertexTypes.h
+++ b/Engine/Plugins/Experimental/ProxyLODPlugin/Source/ProxyLOD/Private/ProxyLODVertexTypes.h
@@ -115,13 +115,6 @@ public:
 	FVector3f			Normal;
 
 
-	FPositionNormalVertex(const FPositionNormalVertex& OtherVert) :
-		MaterialIndex(OtherVert.MaterialIndex),
-		Position(OtherVert.Position),
-		Normal(OtherVert.Normal)
-	{}
-
-
 	// Attributes are all assumed to be floats.  
 	static uint32 NumAttributes()
 	{
