Removed monolothic header includes.

diff --git a/Engine/Plugins/Experimental/ProxyLODPlugin/Source/ProxyLOD/Private/ProxyLODPlugin.cpp b/Engine/Plugins/Experimental/ProxyLODPlugin/Source/ProxyLOD/Private/ProxyLODPlugin.cpp
--- a/Engine/Plugins/Experimental/ProxyLODPlugin/Source/ProxyLOD/Private/ProxyLODPlugin.cpp
+++ b/Engine/Plugins/Experimental/ProxyLODPlugin/Source/ProxyLOD/Private/ProxyLODPlugin.cpp
@@ -1,6 +1,5 @@
 // Copyright Epic Games, Inc. All Rights Reserved.
 
-//#include "Engine.h"
 #include "CoreMinimal.h"
 #include "Modules/ModuleManager.h"
 #include "Features/IModularFeatures.h"
diff --git a/Engine/Source/Developer/TraceAnalysis/Private/Analysis/Engine.cpp b/Engine/Source/Developer/TraceAnalysis/Private/Analysis/Engine.cpp
--- a/Engine/Source/Developer/TraceAnalysis/Private/Analysis/Engine.cpp
+++ b/Engine/Source/Developer/TraceAnalysis/Private/Analysis/Engine.cpp
@@ -1,6 +1,6 @@
 // Copyright Epic Games, Inc. All Rights Reserved.
 
-#include "Engine.h"
+#include "Analysis/Engine.h"
 
 #include "Algo/BinarySearch.h"
 #include "Algo/Sort.h"
diff --git a/Engine/Source/Developer/TraceAnalysis/Private/Analysis/Processor.h b/Engine/Source/Developer/TraceAnalysis/Private/Analysis/Processor.h
--- a/Engine/Source/Developer/TraceAnalysis/Private/Analysis/Processor.h
+++ b/Engine/Source/Developer/TraceAnalysis/Private/Analysis/Processor.h
@@ -2,7 +2,7 @@
 
 #pragma once
 
-#include "Engine.h"
+#include "Analysis/Engine.h"
 #include "HAL/Runnable.h"
 #include "Logging/MessageLog.h"
 #include "Trace/Analysis.h"
