ProxyLOD - Fix corner case where processing was skipped if input geometry had no volume. What we really want to test is for an area.

diff --git a/Engine/Plugins/Editor/ProxyLODPlugin/Source/ProxyLOD/Private/ProxyLODPlugin.cpp b/Engine/Plugins/Editor/ProxyLODPlugin/Source/ProxyLOD/Private/ProxyLODPlugin.cpp
--- a/Engine/Plugins/Editor/ProxyLODPlugin/Source/ProxyLOD/Private/ProxyLODPlugin.cpp
+++ b/Engine/Plugins/Editor/ProxyLODPlugin/Source/ProxyLOD/Private/ProxyLODPlugin.cpp
@@ -676,7 +676,7 @@ void FVoxelizeMeshMerging::ProxyLOD(FMeshDescriptionArrayAdapter& InSrcGeometryA
 		{
 			const auto& BBox = InSrcGeometryAdapter.GetBBox();
 
-			if (!BBox.hasVolume())
+			if (BBoxMajorAxisLength(BBox) == 0.f)
 			{
 				UE_LOG(LogProxyLODMeshReduction, Warning, TEXT("Empty bounding box for all static meshes input, no output mesh will be generated."));
 
