Merge Actor - Simplification: When using the gap filing option, avoid the costly round trip to a FMeshDescription
* HLOD1 are building around 4x faster with this change (building 50 HLOD1 actors went from ~25m to ~6m

diff --git a/Engine/Plugins/Editor/ProxyLODPlugin/Source/ProxyLOD/Private/ProxyLODMeshConvertUtils.h b/Engine/Plugins/Editor/ProxyLODPlugin/Source/ProxyLOD/Private/ProxyLODMeshConvertUtils.h
--- a/Engine/Plugins/Editor/ProxyLODPlugin/Source/ProxyLOD/Private/ProxyLODMeshConvertUtils.h
+++ b/Engine/Plugins/Editor/ProxyLODPlugin/Source/ProxyLOD/Private/ProxyLODMeshConvertUtils.h
@@ -220,45 +220,3 @@ void ProxyLOD::MixedPolyMeshToAOSMesh(const FMixedPolyMesh& MixedPolyMesh, TAOSM
 		});
 	}
 }
-
-
-
-
-
-
-
-#if 0
-static void ComputeAABB(const FMeshDescriptionArrayAdapter& MeshArray, ProxyLOD::FBBox& BBox)
-{
-	uint32 NumTris = MeshArray.polygonCount();
-	BBox = ProxyLOD::Parallel_Reduce(ProxyLOD::FIntRange(0, NumTris), ProxyLOD::FBBox(), [&MeshArray](const ProxyLOD::FIntRange& Range,  ProxyLOD::FBBox BBox)->ProxyLOD::FBBox
-	{
-		// loop over faces
-		for (int32 f = Range.begin(), F = Range.end(); f < F; ++f)
-		{
-			openvdb::Vec3d Pos;
-			// loop over verts
-			for (int32 v = 0; v < 3; ++v)
-			{
-				MeshArray.getWorldSpacePoint(f, v, Pos);
-
-				BBox.expand(Pos);
-			}
-
-		}
-
-		return BBox;
-
-	}, [](const ProxyLOD::FBBox& BBoxA, const ProxyLOD::FBBox& BBoxB)->ProxyLOD::FBBox
-	{
-		ProxyLOD::FBBox Result(BBoxA);
-		Result.expand(BBoxB);
-		
-		return Result;
-	}
-	);
-}
-
-
-#endif 
-
diff --git a/Engine/Plugins/Editor/ProxyLODPlugin/Source/ProxyLOD/Private/ProxyLODMeshSDFConversions.cpp b/Engine/Plugins/Editor/ProxyLODPlugin/Source/ProxyLOD/Private/ProxyLODMeshSDFConversions.cpp
--- a/Engine/Plugins/Editor/ProxyLODPlugin/Source/ProxyLOD/Private/ProxyLODMeshSDFConversions.cpp
+++ b/Engine/Plugins/Editor/ProxyLODPlugin/Source/ProxyLOD/Private/ProxyLODMeshSDFConversions.cpp
@@ -6,94 +6,6 @@
 
 #include "CoreMinimal.h"
 
-//#include <exception>
-
-
-bool ProxyLOD::MeshArrayToSDFVolume(const FMeshDescriptionArrayAdapter& MeshAdapter, openvdb::FloatGrid::Ptr& SDFGrid, openvdb::Int32Grid* PolyIndexGrid)
-{
-	TRACE_CPUPROFILER_EVENT_SCOPE(ProxyLOD::MeshArrayToSDFVolume)
-
-	bool success = true;
-	try
-	{
-		const float HalfBandWidth = 2.f;
-		int Flags = 0;
-		{
-			TRACE_CPUPROFILER_EVENT_SCOPE(OpenVDB::MeshToVolume)
-			SDFGrid = openvdb::tools::meshToVolume<openvdb::FloatGrid>(MeshAdapter, MeshAdapter.GetTransform(), HalfBandWidth /*exterior*/, HalfBandWidth/*interior*/, Flags, PolyIndexGrid);
-		}
-
-		// reduce memory footprint, increase the sparseness.
-		{
-			TRACE_CPUPROFILER_EVENT_SCOPE(OpenVDB::PruneLevelSet)
-			openvdb::tools::pruneLevelSet(SDFGrid->tree(), HalfBandWidth, -HalfBandWidth);
-		}
-	}
-	catch (std::bad_alloc& )
-	{
-		success = false;
-		if (SDFGrid)
-		{
-			SDFGrid->tree().clear(); // free any memory held in the smart pointers in the grid
-		}
-		if (PolyIndexGrid)
-		{
-			PolyIndexGrid->clear();
-		}
-	}
-
-#if 0
-	// Used in testing
-	openvdb::tree::LeafManager<openvdb::FloatTree> leafManager(SDFGrid->tree());
-
-	// The number of 8x8x8 voxel bricks.
-	size_t leafCount = leafManager.leafCount();
-	// number of voxels with real distance values
-	size_t activeCount = SDFGrid->tree().activeLeafVoxelCount();
-	// total amount of memory used.
-	size_t numByets = SDFGrid->tree().memUsage();
-#endif
-
-	return success;
-
-}
-
-
-bool ProxyLOD::MeshArrayToSDFVolume(const FMeshDescriptionAdapter& MeshAdapter, openvdb::FloatGrid::Ptr& SDFGrid, openvdb::Int32Grid* PolyIndexGrid)
-{
-	TRACE_CPUPROFILER_EVENT_SCOPE(ProxyLOD::MeshArrayToSDFVolume)
-
-	bool success = true;
-	try
-	{
-		const float HalfBandWidth = 2.f;
-		int Flags = 0;
-		{
-			TRACE_CPUPROFILER_EVENT_SCOPE(OpenVDB::MeshToVolume)
-			SDFGrid = openvdb::tools::meshToVolume<openvdb::FloatGrid>(MeshAdapter, MeshAdapter.GetTransform(), HalfBandWidth /*exterior*/, HalfBandWidth/*interior*/, Flags, PolyIndexGrid);
-		}
-
-		// reduce memory footprint, increase the spareness.
-		{
-			TRACE_CPUPROFILER_EVENT_SCOPE(OpenVDB::PruneLevelSet)
-			openvdb::tools::pruneLevelSet(SDFGrid->tree(), HalfBandWidth, -HalfBandWidth);
-		}
-	}
-	catch (std::bad_alloc&)
-	{
-		success = false;
-		if (SDFGrid)
-		{
-			SDFGrid->tree().clear(); // free any memory held in the smart pointers in the grid
-		}
-		if (PolyIndexGrid)
-		{
-			PolyIndexGrid->clear();
-		}
-	}
-	return success;
-}
-
 
 /**
 * Generate a new SDF (with narrow band thickness of 2) that represents moving the zero crossing
@@ -109,7 +21,7 @@ bool ProxyLOD::MeshArrayToSDFVolume(const FMeshDescriptionAdapter& MeshAdapter,
 *
 * @return A new SDF that represents a dilation or erosion (expansion or contraction) of the original SDF
 */
-static openvdb::FloatGrid::Ptr OffsetSDF(const openvdb::FloatGrid::Ptr InSDFVolume, const double WSOffset, const double ResultVolexSize)
+static openvdb::FloatGrid::Ptr OffsetSDF(const openvdb::FloatGrid::Ptr InSDFVolume, const double WSOffset, const double ResultVoxelSize)
 {
 	TRACE_CPUPROFILER_EVENT_SCOPE(OffsetSDF)
 	// Extract the iso-surface with offset DilationInVoxels.
@@ -129,20 +41,13 @@ static openvdb::FloatGrid::Ptr OffsetSDF(const openvdb::FloatGrid::Ptr InSDFVolu
 		openvdb::tools::volumeToMesh(*InSDFVolume, MixedPolyMesh.Points, MixedPolyMesh.Triangles, MixedPolyMesh.Quads, IsoValue, 0.001);
 	}
 
-	// convert the mesh to FMeshDescription
-	FMeshDescription RawMesh;
-	FStaticMeshAttributes(RawMesh).Register();
-	ProxyLOD::MixedPolyMeshToRawMesh(MixedPolyMesh, RawMesh);
-
 	// Create a new empty grid with the same transform and metadata
 	openvdb::FloatGrid::Ptr OutSDFVolume = openvdb::FloatGrid::create(*InSDFVolume);
-	OutSDFVolume->setTransform(openvdb::math::Transform::createLinearTransform(ResultVolexSize));
-
-	// Wrap so we can re-voxelize with bandwidth 2
-	FMeshDescriptionAdapter  MeshAdapter(RawMesh, OutSDFVolume->transform());
+	OutSDFVolume->setTransform(openvdb::math::Transform::createLinearTransform(ResultVoxelSize));
 
-	// Re-voxelize
-	ProxyLOD::MeshArrayToSDFVolume(MeshAdapter, OutSDFVolume);
+	// Re-voxelize with bandwidth 2
+	MixedPolyMesh.Transform = OutSDFVolume->transform();
+	ProxyLOD::MeshToSDFVolume(MixedPolyMesh, OutSDFVolume->transform(), OutSDFVolume);
 
 	return OutSDFVolume;
 }
@@ -150,7 +55,7 @@ static openvdb::FloatGrid::Ptr OffsetSDF(const openvdb::FloatGrid::Ptr InSDFVolu
 
 void ProxyLOD::CloseGaps(openvdb::FloatGrid::Ptr InOutSDFVolume, const double GapRadius, const int32 MaxDilations)
 {
-	// Implimentaiton notes:
+	// Implementation notes:
 	// This functions by first inflating (dilate) the geometry SDF (moving the surface outward along the normal) an amount 
 	// GapRadius.  Doing this may bring surfaces into contact, thus closing gaps.
 	// Next the geometry SDF with merged gaps is deflated (erode) to a size that should be slightly smaller than the original geometry.
diff --git a/Engine/Plugins/Editor/ProxyLODPlugin/Source/ProxyLOD/Private/ProxyLODMeshSDFConversions.h b/Engine/Plugins/Editor/ProxyLODPlugin/Source/ProxyLOD/Private/ProxyLODMeshSDFConversions.h
--- a/Engine/Plugins/Editor/ProxyLODPlugin/Source/ProxyLOD/Private/ProxyLODMeshSDFConversions.h
+++ b/Engine/Plugins/Editor/ProxyLODPlugin/Source/ProxyLOD/Private/ProxyLODMeshSDFConversions.h
@@ -12,10 +12,10 @@ namespace ProxyLOD
 {
 
 	/**
-	* Generate a single sparse Signed Distance Field (SDF) Grid from the polys managed by MeshAdapter,  
+	* Generate a single sparse Signed Distance Field (SDF) Grid from the polys managed by InMesh,  
 	* Optionally store the index of the closest poly is recorded in a sparse PolyIndexGrid.  
 	*
-	* @param  InMeshAdapter     Source polygons to be voxelized.
+	* @param  InMesh			Source polygons to be voxelized.
 	* @param  OutSDFGrid        Resulting sparse signed distance field.
 	* @param  OutPolyIndexGrid  Optional resulting sparse grid. For each voxel with valid distance in the OutSDFGrid
 	*                           the index of the closest poly will be recorded in the corresponding voxel 
@@ -23,13 +23,11 @@ namespace ProxyLOD
 	*
 	* @retun 'true' is success, 'false' will generally be an out of memory error.
 	*/
-	bool MeshArrayToSDFVolume( const FMeshDescriptionArrayAdapter& InMeshAdapter,
-	                           openvdb::FloatGrid::Ptr& OutSDFGrid,
-		                       openvdb::Int32Grid* OutPolyIndexGrid = nullptr);
-
-	bool MeshArrayToSDFVolume( const FMeshDescriptionAdapter& InMeshAdapter,
-		                       openvdb::FloatGrid::Ptr& OutSDFGrid,
-		                       openvdb::Int32Grid* OutPolyIndexGrid = nullptr);
+	template <typename SrcMeshType>
+	bool MeshToSDFVolume( const SrcMeshType& InMesh,
+						  const openvdb::math::Transform& InTransform,
+	                      openvdb::FloatGrid::Ptr& OutSDFGrid,
+		                  openvdb::Int32Grid* OutPolyIndexGrid = nullptr);
 
 	/**
 	* Extract the isosurface (=IsoValue) of the given SDF volume, in the form of triangle mesh.
@@ -99,7 +97,42 @@ namespace ProxyLOD
 
 }
 
-
+template <typename SrcMeshType>
+bool ProxyLOD::MeshToSDFVolume(const SrcMeshType& InMesh, const openvdb::math::Transform& InTransform, openvdb::FloatGrid::Ptr& SDFGrid, openvdb::Int32Grid* PolyIndexGrid)
+{
+	TRACE_CPUPROFILER_EVENT_SCOPE(ProxyLOD::MeshToSDFVolume)
+
+	bool bSuccess = true;
+	try
+	{
+		const float HalfBandWidth = 2.f;
+		int Flags = 0;
+		{
+			TRACE_CPUPROFILER_EVENT_SCOPE(OpenVDB::MeshToVolume)
+			SDFGrid = openvdb::tools::meshToVolume<openvdb::FloatGrid>(InMesh, InTransform, HalfBandWidth /*exterior*/, HalfBandWidth/*interior*/, Flags, PolyIndexGrid);
+		}
+
+		// reduce memory footprint, increase the sparseness.
+		{
+			TRACE_CPUPROFILER_EVENT_SCOPE(OpenVDB::PruneLevelSet)
+			openvdb::tools::pruneLevelSet(SDFGrid->tree(), HalfBandWidth, -HalfBandWidth);
+		}
+	}
+	catch (std::bad_alloc&)
+	{
+		bSuccess = false;
+		if (SDFGrid)
+		{
+			SDFGrid->tree().clear(); // free any memory held in the smart pointers in the grid
+		}
+		if (PolyIndexGrid)
+		{
+			PolyIndexGrid->clear();
+		}
+	}
+
+	return bSuccess;
+}
 
 
 template <typename DstMeshType>
@@ -118,9 +151,6 @@ void ProxyLOD::SDFVolumeToMesh(const openvdb::FloatGrid::ConstPtr SDFVolume, con
 	// Convert to Quad Mesh - extracting iso-surface defined by IsoValue
 
 	openvdb::tools::volumeToMesh(*SDFVolume, TempMixedPolyMesh.Points, TempMixedPolyMesh.Triangles, TempMixedPolyMesh.Quads, IsoValue, Adaptivity);
-
-	//TestUniqueVertexes(SimpleMesh);
-
 	
 	MixedPolyMeshToAOSMesh(TempMixedPolyMesh, OutMesh);
 }
diff --git a/Engine/Plugins/Editor/ProxyLODPlugin/Source/ProxyLOD/Private/ProxyLODMeshTypes.h b/Engine/Plugins/Editor/ProxyLODPlugin/Source/ProxyLOD/Private/ProxyLODMeshTypes.h
--- a/Engine/Plugins/Editor/ProxyLODPlugin/Source/ProxyLOD/Private/ProxyLODMeshTypes.h
+++ b/Engine/Plugins/Editor/ProxyLODPlugin/Source/ProxyLOD/Private/ProxyLODMeshTypes.h
@@ -83,15 +83,74 @@ void ResizeInializedArray(TArray<ValueType>& Array, int32 Size)
 }
 /**
 * Mesh type that holds minimal required data for openvdb iso-surface extraction 
-* interface.
+* interface. Implements the openvdb MeshDataAdapter interface.
 * 
-* NB: std::vector required by openvdb interface.
+* NB: std::vector required by the openvdb interface.
 */
 struct FMixedPolyMesh
 {
 	std::vector<openvdb::Vec3s> Points;
 	std::vector<openvdb::Vec4I> Quads;
 	std::vector<openvdb::Vec3I> Triangles;
+	openvdb::math::Transform	Transform;
+
+	// ~MeshDataAdapter Interface Begin
+	// https://www.openvdb.org/documentation/doxygen/interfaceMeshDataAdapter.html
+	
+	// Total number of polygons
+	size_t polygonCount() const
+	{
+		return Quads.size() + Triangles.size();
+	}
+
+	// Total number of points (vertex locations)
+	size_t pointCount() const
+	{
+		return Points.size();
+	}
+
+	// Vertex count for polygon n
+	size_t vertexCount(size_t n) const
+	{
+		return PolygonIsQuad(n) ? 4 : 3;
+	}
+
+	// Return position pos in local grid index space for polygon n and vertex v
+	void getIndexSpacePoint(size_t n, size_t v, openvdb::Vec3d& pos) const
+	{
+		if (PolygonIsQuad(n))
+		{
+			pos = Points[Quads[GetQuadIndex(n)][v]];
+		}
+		else
+		{
+			pos = Points[Triangles[GetTriangleIndex(n)][v]];
+		}
+
+		pos = Transform.worldToIndex(pos);
+	}
+	// ~MeshDataAdapter Interface End
+
+private:
+	bool PolygonIsQuad(size_t n) const
+	{
+		return n < Quads.size();
+	}
+
+	bool PolygonIsTriangle(size_t n) const
+	{
+		return n >= Triangles.size();
+	}
+
+	size_t GetQuadIndex(size_t n) const
+	{
+		return n;
+	}
+
+	size_t GetTriangleIndex(size_t n) const
+	{
+		return n - Quads.size();
+	}
 };
 
 /**
diff --git a/Engine/Plugins/Editor/ProxyLODPlugin/Source/ProxyLOD/Private/ProxyLODPlugin.cpp b/Engine/Plugins/Editor/ProxyLODPlugin/Source/ProxyLOD/Private/ProxyLODPlugin.cpp
--- a/Engine/Plugins/Editor/ProxyLODPlugin/Source/ProxyLOD/Private/ProxyLODPlugin.cpp
+++ b/Engine/Plugins/Editor/ProxyLODPlugin/Source/ProxyLOD/Private/ProxyLODPlugin.cpp
@@ -754,7 +754,7 @@ void FVoxelizeMeshMerging::ProxyLOD(FMeshDescriptionArrayAdapter& InSrcGeometryA
 
 				// 1) Voxelize - this can potentially run out of memory when very large objecs (or very small voxel sizes) are used.
 				
-				const bool bSuccess = ProxyLOD::MeshArrayToSDFVolume(InSrcGeometryAdapter, SDFVolume, SrcPolyIndexGrid.get());
+				const bool bSuccess = ProxyLOD::MeshToSDFVolume(InSrcGeometryAdapter, InSrcGeometryAdapter.GetTransform(), SDFVolume, SrcPolyIndexGrid.get());
 				const bool bHasClipping = InClippingGeometryAdapter.polygonCount() != 0;
 
 				if (bSuccess && bHasClipping)
@@ -762,7 +762,7 @@ void FVoxelizeMeshMerging::ProxyLOD(FMeshDescriptionArrayAdapter& InSrcGeometryA
 					// Voxelize the clipping geometry.
 
 					openvdb::FloatGrid::Ptr SDFClipping;
-					ProxyLOD::MeshArrayToSDFVolume(InClippingGeometryAdapter, SDFClipping);
+					ProxyLOD::MeshToSDFVolume(InClippingGeometryAdapter, InClippingGeometryAdapter.GetTransform(), SDFClipping);
 
 					// CSG difference that removes the clipping region from the SDFvolume, leaving a watertight SDF
 
diff --git a/Engine/Plugins/Editor/ProxyLODPlugin/Source/ProxyLOD/Private/ProxyLODVolume.cpp b/Engine/Plugins/Editor/ProxyLODPlugin/Source/ProxyLOD/Private/ProxyLODVolume.cpp
--- a/Engine/Plugins/Editor/ProxyLODPlugin/Source/ProxyLOD/Private/ProxyLODVolume.cpp
+++ b/Engine/Plugins/Editor/ProxyLODPlugin/Source/ProxyLOD/Private/ProxyLODVolume.cpp
@@ -128,7 +128,7 @@ private:
 
 		SrcPolyIndexGrid = openvdb::Int32Grid::create();
 
-		if (!ProxyLOD::MeshArrayToSDFVolume(InSrcGeometryAdapter, SDFVolume, SrcPolyIndexGrid.get()))
+		if (!ProxyLOD::MeshToSDFVolume(InSrcGeometryAdapter, InSrcGeometryAdapter.GetTransform(), SDFVolume, SrcPolyIndexGrid.get()))
 		{
 			SrcPolyIndexGrid.reset();
 			return false;
