PR #12362: Address some clang warnings

diff --git a/Engine/Plugins/Editor/ProxyLODPlugin/Source/ProxyLOD/Private/ProxyLODMeshUtilities.cpp b/Engine/Plugins/Editor/ProxyLODPlugin/Source/ProxyLOD/Private/ProxyLODMeshUtilities.cpp
--- a/Engine/Plugins/Editor/ProxyLODPlugin/Source/ProxyLOD/Private/ProxyLODMeshUtilities.cpp
+++ b/Engine/Plugins/Editor/ProxyLODPlugin/Source/ProxyLOD/Private/ProxyLODMeshUtilities.cpp
@@ -1076,7 +1076,7 @@ int ProxyLOD::VerifyAdjacency(const uint32* EdgeAdjacentFaceArray, const uint32
 			const uint32 AdjAdjFaces[3] = { EdgeAdjacentFaceArray[AdjFaceIdx], EdgeAdjacentFaceArray[AdjFaceIdx + 1], EdgeAdjacentFaceArray[AdjFaceIdx + 2] };
 
 			// one of these should be 'f' itself.
-			bValid = bValid && AdjAdjFaces[0] == f || AdjAdjFaces[1] == f || AdjAdjFaces[2] == f;
+			bValid = bValid && (AdjAdjFaces[0] == f || AdjAdjFaces[1] == f || AdjAdjFaces[2] == f);
 
 			if (!bValid)
 			{
diff --git a/Engine/Source/Runtime/Core/Private/Windows/WindowsPlatformMisc.cpp b/Engine/Source/Runtime/Core/Private/Windows/WindowsPlatformMisc.cpp
--- a/Engine/Source/Runtime/Core/Private/Windows/WindowsPlatformMisc.cpp
+++ b/Engine/Source/Runtime/Core/Private/Windows/WindowsPlatformMisc.cpp
@@ -2376,7 +2376,7 @@ bool FWindowsPlatformMisc::IsValidAbsolutePathFormat(const FString& Path)
 	}
 
 	// Find any unnamed directory changes
-	if ( bIsValid && (Path.Find(TEXT("//"))!=INDEX_NONE) || (Path.Find(TEXT("\\/"))!=INDEX_NONE) || (Path.Find(TEXT("/\\"))!=INDEX_NONE) || (Path.Find(TEXT("\\\\"))!=INDEX_NONE) )
+	if ( bIsValid && (Path.Find(TEXT("//"))!=INDEX_NONE || Path.Find(TEXT("\\/"))!=INDEX_NONE || Path.Find(TEXT("/\\"))!=INDEX_NONE || Path.Find(TEXT("\\\\"))!=INDEX_NONE) )
 	{
 		bIsValid = false;
 	}
diff --git a/Engine/Source/Runtime/D3D12RHI/Private/D3D12Allocation.cpp b/Engine/Source/Runtime/D3D12RHI/Private/D3D12Allocation.cpp
--- a/Engine/Source/Runtime/D3D12RHI/Private/D3D12Allocation.cpp
+++ b/Engine/Source/Runtime/D3D12RHI/Private/D3D12Allocation.cpp
@@ -872,7 +872,7 @@ FD3D12BucketAllocator::FD3D12BucketAllocator(FD3D12Device* ParentDevice,
 	const FD3D12ResourceInitConfig& InInitConfig,
 	const FString& Name,
 	uint64 InBlockRetentionFrameCount) :
-	FD3D12ResourceAllocator(ParentDevice, VisibleNodes, InitConfig, Name, 32 * 1024 * 1024),
+	FD3D12ResourceAllocator(ParentDevice, VisibleNodes, InInitConfig, Name, 32 * 1024 * 1024),
 	BlockRetentionFrameCount(InBlockRetentionFrameCount)
 {}
 
