Fix windows definition leaks causing issues with unity files

#rnx

diff --git a/Engine/Plugins/Editor/ProxyLODPlugin/Source/ProxyLOD/Private/ProxyLODOpenVDB.h b/Engine/Plugins/Editor/ProxyLODPlugin/Source/ProxyLOD/Private/ProxyLODOpenVDB.h
--- a/Engine/Plugins/Editor/ProxyLODPlugin/Source/ProxyLOD/Private/ProxyLODOpenVDB.h
+++ b/Engine/Plugins/Editor/ProxyLODPlugin/Source/ProxyLOD/Private/ProxyLODOpenVDB.h
@@ -32,6 +32,7 @@ UE_PUSH_MACRO("check")
 #include <openvdb/tools/Interpolation.h> 	// for grid sampler
 #include <openvdb/tools/MeshToVolume.h> 	// for MeshToVolume
 #include <openvdb/tools/VolumeToMesh.h> 	// for VolumeToMesh
+#undef GetObject
 
 #if LOCAL_M_PI
 	#undef M_PI
