Workaround to prevent crash in UVAtlas

diff --git a/Engine/Plugins/Editor/ProxyLODPlugin/Source/ProxyLOD/Private/ProxyLODMeshParameterization.cpp b/Engine/Plugins/Editor/ProxyLODPlugin/Source/ProxyLOD/Private/ProxyLODMeshParameterization.cpp
--- a/Engine/Plugins/Editor/ProxyLODPlugin/Source/ProxyLOD/Private/ProxyLODMeshParameterization.cpp
+++ b/Engine/Plugins/Editor/ProxyLODPlugin/Source/ProxyLOD/Private/ProxyLODMeshParameterization.cpp
@@ -81,15 +81,15 @@ bool ProxyLOD::GenerateUVs(const FTextureAtlasDesc& TextureAtlasDesc,
 	const float  Gutter = TextureAtlasDesc.Gutter;
 
 	// Symmetric Identity matrix used for the signal
-	float * pIMTArray = new float[NumFaces * 3];
+	std::vector<float> IMTArray(NumFaces * 3);
 	{
 		for (int32 f = 0; f < NumFaces; ++f)
 		{
 			int32 offset = 3 * f;
 			{
-				pIMTArray[offset + 0] = 1.f;
-				pIMTArray[offset + 1] = 0.f;
-				pIMTArray[offset + 2] = 1.f;
+				IMTArray[offset + 0] = 1.f;
+				IMTArray[offset + 1] = 0.f;
+				IMTArray[offset + 2] = 1.f;
 			}
 		}
 	}
@@ -105,22 +105,133 @@ bool ProxyLOD::GenerateUVs(const FTextureAtlasDesc& TextureAtlasDesc,
 	std::vector<uint32> FacePartitioning;
 
 	// Generate UVs
-	HRESULT hr = DirectX::UVAtlasCreate(PosArray.GetData(), NumVerts,
-		                                Indices.GetData(), DXGI_FORMAT_R32_UINT, NumFaces,
-		                                NumChartsIn, MaxStretchIn,
-		                                Width, Height, Gutter,
-		                                AdjacencyArray.GetData(), NULL /*false adj*/, pIMTArray /*IMTArray*/,
-		                                StatusCallBack, DirectX::UVATLAS_DEFAULT_CALLBACK_FREQUENCY,
-		                                DirectX::UVATLAS_DEFAULT, VB, IB,
-		                                &FacePartitioning, &RemapArray, &MaxStretchOut, &NumChartsOut);
+	constexpr DXGI_FORMAT IndexFormat = DXGI_FORMAT_R32_UINT;
+
+	std::vector<uint32_t> vPartitionResultAdjacency;
+
+	HRESULT hr = DirectX::UVAtlasPartition(PosArray.GetData(), NumVerts,
+		Indices.GetData(), IndexFormat, NumFaces,
+		NumChartsIn, MaxStretchIn,
+		AdjacencyArray.GetData(), NULL /*false adj*/, IMTArray.data() /*IMTArray*/,
+		StatusCallBack, DirectX::UVATLAS_DEFAULT_CALLBACK_FREQUENCY,
+		DirectX::UVATLAS_DEFAULT, VB, IB,
+		&FacePartitioning, &RemapArray, vPartitionResultAdjacency, &MaxStretchOut, &NumChartsOut);
+
+	if (hr != S_OK)
+	{
+		return false;
+	}
+
+	if (!ensure(vPartitionResultAdjacency.size() == 3 * NumFaces))
+	{
+		return false;
+	}
+
+	//
+	// UVAtlas has a bug where it can crash if multiple charts reference the same vertex index, depending on the order in which the charts are processed internally. 
+	// (e.g. while processing chart[i], it assumes all vertex indices are in some range [min[i], max[i]]. If it tries to reuse a vertex index less than min[i], it may crash.)
+	// To workaround this, we will duplicate shared vertices into their own charts.
+	//
+
+	std::vector<bool> bUsedFace(NumFaces, false);
+
+	std::vector<DirectX::UVAtlasVertex> NewVertexBuffer;
+
+	static_assert(IndexFormat == DXGI_FORMAT_R32_UINT);
+	check(IB.size() % sizeof(uint32_t) == 0);
+
+	std::vector<uint32_t> IndexBufferU32;
+	IndexBufferU32.resize(IB.size() / sizeof(uint32_t));
+	FMemory::Memcpy(IndexBufferU32.data(), IB.data(), IndexBufferU32.size() * sizeof(uint32_t));
+
+	std::vector<uint32_t> NewIndexBufferU32;
+	NewIndexBufferU32.resize(3 * NumFaces);
+
+	std::vector<uint32> NewRemapArray;
+	NewRemapArray.resize(NumVerts, -1);
+
+	// Build partioning based on triangle adjacency (this is how UVAtlas will initialize things internally, NOT via the output FacePartitioning from UVAtlasPartition as one might expect)
+
+	for (int32 FaceIndex = 0; FaceIndex < NumFaces; FaceIndex++)
+	{
+		if (!bUsedFace[FaceIndex])
+		{
+			// new chart
+			bUsedFace[FaceIndex] = true;
+
+			std::vector<uint32_t> CurrentChartTriangles;
+			CurrentChartTriangles.push_back(FaceIndex);
+
+			size_t ChartTriangleIndex = 0;
+
+			// use breadth-first search to find chart
+			while (ChartTriangleIndex < CurrentChartTriangles.size())
+			{
+				for (uint32_t TriEdgeIndex = 0; TriEdgeIndex < 3; TriEdgeIndex++)
+				{
+					const uint32_t TriangleAdjacencyIndex = 3 * CurrentChartTriangles[ChartTriangleIndex] + TriEdgeIndex;
+					const uint32_t AdjacentTriangleIndex = vPartitionResultAdjacency[TriangleAdjacencyIndex];
+					if (AdjacentTriangleIndex != -1 && !bUsedFace[AdjacentTriangleIndex])
+					{
+						CurrentChartTriangles.push_back(AdjacentTriangleIndex);
+						bUsedFace[AdjacentTriangleIndex] = true;
+					}
+				}
+				ChartTriangleIndex++;
+			}
+
+			// Now rebuild the vertex and index buffers for this chart. Vertices that appear in multiple charts should be duplicated by this process.
+			const size_t ChartMinIndex = NewVertexBuffer.size();
+			TMap<uint32_t, uint32_t> VertexMap;
+
+			for (uint32_t CurrentChartFaceIndex : CurrentChartTriangles)
+			{
+				for (int32 TriVertIndex = 0; TriVertIndex < 3; ++TriVertIndex)
+				{
+					const uint32_t TriVert = IndexBufferU32[3 * CurrentChartFaceIndex + TriVertIndex];
+
+					uint32_t NewVertexIndex;
+					if (uint32_t* MappedVertex = VertexMap.Find(TriVert))
+					{
+						NewVertexIndex = *MappedVertex;
+					}
+					else
+					{
+						NewVertexIndex = (uint32_t)NewVertexBuffer.size();
+						NewVertexBuffer.push_back(VB[TriVert]);
+						VertexMap.Add(TriVert, NewVertexIndex);
+					}
+
+					if (NewRemapArray.size() <= NewVertexIndex)
+					{
+						NewRemapArray.resize(NewVertexIndex + 1);
+					}
+					NewRemapArray[NewVertexIndex] = RemapArray[TriVert];
+
+					ensure(ChartMinIndex <= NewVertexIndex);
+
+					NewIndexBufferU32[3 * CurrentChartFaceIndex + TriVertIndex] = NewVertexIndex;
+				}
+			}
+		}
+	}
+
+	std::vector<uint8> NewIndexBuffer;
+	NewIndexBuffer.resize(NewIndexBufferU32.size() * sizeof(uint32_t));
+	FMemory::Memcpy(NewIndexBuffer.data(), NewIndexBufferU32.data(), NewIndexBuffer.size());
+
+	hr = DirectX::UVAtlasPack(NewVertexBuffer, NewIndexBuffer, DXGI_FORMAT_R32_UINT,
+		Width, Height, Gutter,
+		vPartitionResultAdjacency,
+		StatusCallBack, DirectX::UVATLAS_DEFAULT_CALLBACK_FREQUENCY);
 
 
 	// Translate results to output form.
 	if (hr == S_OK)
 	{
-		const int32 NumUVverts    = (int32)VB.size();
-		const int32 NumUVindices  = (int32)IB.size();
-		const int32 NumRemapArray = (int32)RemapArray.size();
+		const int32 NumUVverts    = (int32)NewVertexBuffer.size();
+		const int32 NumUVindices  = (int32)NewIndexBuffer.size();
+		const int32 NumRemapArray = (int32)NewRemapArray.size();
 
 		// empty and resize
 
@@ -128,7 +239,7 @@ bool ProxyLOD::GenerateUVs(const FTextureAtlasDesc& TextureAtlasDesc,
 		UVIndexBuffer.Empty(NumUVindices);
 		VertexRemapArray.Empty(NumRemapArray);
 
-		for (const DirectX::UVAtlasVertex& UVvertex : VB)
+		for (const DirectX::UVAtlasVertex& UVvertex : NewVertexBuffer)
 		{
 			const auto& UV = UVvertex.uv;
 			UVVertexBuffer.Emplace(FVector2D(UV.x, UV.y));
@@ -138,7 +249,7 @@ bool ProxyLOD::GenerateUVs(const FTextureAtlasDesc& TextureAtlasDesc,
 		{
 			std::vector<uint32> indices;
 			indices.resize(3 * NumFaces);
-			std::memcpy(indices.data(), IB.data(), sizeof(uint32) * 3 * NumFaces);
+			std::memcpy(indices.data(), NewIndexBuffer.data(), sizeof(uint32) * 3 * NumFaces);
 
 			for (uint32 i : indices)
 			{
@@ -146,7 +257,7 @@ bool ProxyLOD::GenerateUVs(const FTextureAtlasDesc& TextureAtlasDesc,
 			}
 		}
 
-		for (const uint32 i : RemapArray)
+		for (const uint32 i : NewRemapArray)
 		{
 			int32 r = i;
 			VertexRemapArray.Add(r);
